Add program for downloading Git dependencies from S3.

diff --git a/Engine/Source/Programs/GitDependencies/App.config b/Engine/Source/Programs/GitDependencies/App.config
new file mode 100644
--- /dev/null
+++ b/Engine/Source/Programs/GitDependencies/App.config
@@ -0,0 +1,6 @@
+﻿<?xml version="1.0" encoding="utf-8" ?>
+<configuration>
+    <startup> 
+        <supportedRuntime version="v4.0" sku=".NETFramework,Version=v4.5" />
+    </startup>
+</configuration>
\ No newline at end of file
diff --git a/Engine/Source/Programs/GitDependencies/DependencyManifest.cs b/Engine/Source/Programs/GitDependencies/DependencyManifest.cs
new file mode 100644
--- /dev/null
+++ b/Engine/Source/Programs/GitDependencies/DependencyManifest.cs
@@ -0,0 +1,65 @@
+﻿using System;
+using System.Collections.Generic;
+using System.Diagnostics;
+using System.IO;
+using System.Linq;
+using System.Text;
+using System.Threading.Tasks;
+using System.Xml.Serialization;
+
+namespace GitDependencies
+{
+	[DebuggerDisplay("{Name}")]
+	public class DependencyFile
+	{
+		[XmlAttribute]
+		public string Name;
+
+		[XmlAttribute]
+		public string Hash;
+	}
+
+	[DebuggerDisplay("{Hash}")]
+	public class DependencyBlob
+	{
+		[XmlAttribute]
+		public string Hash;
+
+		[XmlAttribute]
+		public long Size;
+
+		[XmlAttribute]
+		public string PackHash;
+
+		[XmlAttribute]
+		public long PackOffset;
+	}
+
+	[DebuggerDisplay("{Hash}")]
+	public class DependencyPack
+	{
+		[XmlAttribute]
+		public string Hash;
+
+		[XmlAttribute]
+		public long Size;
+
+		[XmlAttribute]
+		public long CompressedSize;
+
+		[XmlAttribute]
+		public string RemotePath;
+	}
+
+	public class DependencyManifest
+	{
+		[XmlArrayItem("File")]
+		public DependencyFile[] Files = new DependencyFile[0];
+
+		[XmlArrayItem("Blob")]
+		public DependencyBlob[] Blobs = new DependencyBlob[0];
+
+		[XmlArrayItem("Pack")]
+		public DependencyPack[] Packs = new DependencyPack[0];
+	}
+}
diff --git a/Engine/Source/Programs/GitDependencies/GitDependencies.csproj b/Engine/Source/Programs/GitDependencies/GitDependencies.csproj
new file mode 100644
--- /dev/null
+++ b/Engine/Source/Programs/GitDependencies/GitDependencies.csproj
@@ -0,0 +1,65 @@
+﻿<?xml version="1.0" encoding="utf-8"?>
+<Project ToolsVersion="12.0" DefaultTargets="Build" xmlns="http://schemas.microsoft.com/developer/msbuild/2003">
+  <Import Project="$(MSBuildExtensionsPath)\$(MSBuildToolsVersion)\Microsoft.Common.props" Condition="Exists('$(MSBuildExtensionsPath)\$(MSBuildToolsVersion)\Microsoft.Common.props')" />
+  <PropertyGroup>
+    <Configuration Condition=" '$(Configuration)' == '' ">Debug</Configuration>
+    <Platform Condition=" '$(Platform)' == '' ">AnyCPU</Platform>
+    <ProjectGuid>{9B740143-841C-4AF7-905E-D8883F1735DE}</ProjectGuid>
+    <OutputType>Exe</OutputType>
+    <AppDesignerFolder>Properties</AppDesignerFolder>
+    <RootNamespace>GitDependencies</RootNamespace>
+    <AssemblyName>GitDependencies</AssemblyName>
+    <TargetFrameworkVersion>v4.5</TargetFrameworkVersion>
+    <FileAlignment>512</FileAlignment>
+  </PropertyGroup>
+  <PropertyGroup Condition=" '$(Configuration)|$(Platform)' == 'Debug|AnyCPU' ">
+    <PlatformTarget>AnyCPU</PlatformTarget>
+    <DebugSymbols>true</DebugSymbols>
+    <DebugType>full</DebugType>
+    <Optimize>false</Optimize>
+    <OutputPath>..\..\..\Binaries\DotNET\</OutputPath>
+    <DefineConstants>DEBUG;TRACE</DefineConstants>
+    <ErrorReport>prompt</ErrorReport>
+    <WarningLevel>4</WarningLevel>
+  </PropertyGroup>
+  <PropertyGroup Condition=" '$(Configuration)|$(Platform)' == 'Development|AnyCPU' ">
+    <PlatformTarget>AnyCPU</PlatformTarget>
+    <DebugType>pdbonly</DebugType>
+    <Optimize>false</Optimize>
+    <OutputPath>..\..\..\Binaries\DotNET\</OutputPath>
+    <DefineConstants>TRACE</DefineConstants>
+    <ErrorReport>prompt</ErrorReport>
+    <WarningLevel>4</WarningLevel>
+  </PropertyGroup>
+  <ItemGroup>
+    <Reference Include="System" />
+    <Reference Include="System.Core" />
+    <Reference Include="System.Xml.Linq" />
+    <Reference Include="System.Data.DataSetExtensions" />
+    <Reference Include="Microsoft.CSharp" />
+    <Reference Include="System.Data" />
+    <Reference Include="System.Xml" />
+  </ItemGroup>
+  <ItemGroup>
+    <Compile Include="..\DotNETCommon\MetaData.cs">
+      <Link>Properties\MetaData.cs</Link>
+    </Compile>
+    <Compile Include="DependencyManifest.cs" />
+    <Compile Include="Log.cs" />
+    <Compile Include="NotifyReadStream.cs" />
+    <Compile Include="Program.cs" />
+    <Compile Include="Properties\AssemblyInfo.cs" />
+    <Compile Include="WorkingManifest.cs" />
+  </ItemGroup>
+  <ItemGroup>
+    <None Include="App.config" />
+  </ItemGroup>
+  <Import Project="$(MSBuildToolsPath)\Microsoft.CSharp.targets" />
+  <!-- To modify your build process, add your task inside one of the targets below and uncomment it. 
+       Other similar extension points exist, see Microsoft.Common.targets.
+  <Target Name="BeforeBuild">
+  </Target>
+  <Target Name="AfterBuild">
+  </Target>
+  -->
+</Project>
\ No newline at end of file
diff --git a/Engine/Source/Programs/GitDependencies/Log.cs b/Engine/Source/Programs/GitDependencies/Log.cs
new file mode 100644
--- /dev/null
+++ b/Engine/Source/Programs/GitDependencies/Log.cs
@@ -0,0 +1,74 @@
+﻿using System;
+using System.Collections.Generic;
+using System.Linq;
+using System.Text;
+using System.Threading.Tasks;
+
+namespace GitDependencies
+{
+	static class Log
+	{
+		static string CurrentStatus = "";
+
+		public static void WriteLine()
+		{
+			FlushStatus();
+			Console.WriteLine();
+		}
+
+		public static void WriteLine(string Line)
+		{
+			FlushStatus();
+			Console.WriteLine(Line);
+		}
+
+		public static void WriteLine(string Format, params object[] Args)
+		{
+			FlushStatus();
+			Console.WriteLine(Format, Args);
+		}
+
+		public static void WriteError(string Format, params object[] Args)
+		{
+			FlushStatus();
+			Console.ForegroundColor = ConsoleColor.Red;
+			Console.WriteLine(Format, Args);
+			Console.ResetColor();
+		}
+
+		public static void WriteStatus(string Format, params object[] Args)
+		{
+			string NewStatus = String.Format(Format, Args);
+
+			// Figure out how many characters are in common
+			int NumCommon = 0;
+			while(NumCommon < CurrentStatus.Length && NumCommon < NewStatus.Length && CurrentStatus[NumCommon] == NewStatus[NumCommon])
+			{
+				NumCommon++;
+			}
+
+			// Go back to the point that they diverge and write the new string
+			Console.Write(new String('\b', CurrentStatus.Length - NumCommon));
+			Console.Write(NewStatus.Substring(NumCommon));
+
+			// If the new message is shorter, clear the rest of the line
+			if(NewStatus.Length < CurrentStatus.Length)
+			{
+				Console.Write(new String(' ', CurrentStatus.Length - NewStatus.Length));
+				Console.Write(new String('\b', CurrentStatus.Length - NewStatus.Length));
+			}
+
+			// Update the status message
+			CurrentStatus = NewStatus;
+		}
+
+		public static void FlushStatus()
+		{
+			if(CurrentStatus.Length > 0)
+			{
+				Console.WriteLine();
+				CurrentStatus = "";
+			}
+		}
+	}
+}
diff --git a/Engine/Source/Programs/GitDependencies/NotifyReadStream.cs b/Engine/Source/Programs/GitDependencies/NotifyReadStream.cs
new file mode 100644
--- /dev/null
+++ b/Engine/Source/Programs/GitDependencies/NotifyReadStream.cs
@@ -0,0 +1,91 @@
+﻿using System;
+using System.Collections.Generic;
+using System.IO;
+using System.Linq;
+using System.Text;
+using System.Threading.Tasks;
+
+namespace GitDependencies
+{
+	delegate void NotifyReadDelegate(int SizeRead);
+
+	class NotifyReadStream : Stream
+	{
+		Stream Inner;
+		NotifyReadDelegate NotifyRead;
+
+		public NotifyReadStream(Stream InInner, NotifyReadDelegate InNotifyRead)
+		{
+			Inner = InInner;
+			NotifyRead = InNotifyRead;
+		}
+
+		protected override void Dispose(bool Disposing)
+		{
+			if(Inner != null)
+			{
+				Inner.Dispose();
+				Inner = null;
+			}
+		}
+
+		public override bool CanRead
+		{
+			get { return true; }
+		}
+
+		public override bool CanWrite
+		{
+			get { return false; }
+		}
+
+		public override bool CanSeek
+		{
+			get { return false; }
+		}
+
+		public override long Position
+		{
+			get
+			{
+				return Inner.Position;
+			}
+			set
+			{
+				throw new NotImplementedException();
+			}
+		}
+
+		public override long Length
+		{
+			get { return Inner.Length; }
+		}
+
+		public override void SetLength(long Value)
+		{
+			throw new NotImplementedException();
+		}
+
+		public override int Read(byte[] Buffer, int Offset, int Count)
+		{
+			int SizeRead = Inner.Read(Buffer, Offset, Count);
+			NotifyRead(SizeRead);
+			return SizeRead;
+		}
+
+		public override void Write(byte[] Buffer, int Offset, int Count)
+		{
+			throw new NotImplementedException();
+		}
+
+		public override long Seek(long Offset, SeekOrigin Origin)
+		{
+			throw new NotImplementedException();
+		}
+
+		public override void Flush()
+		{
+			Inner.Flush();
+		}
+	}
+}
diff --git a/Engine/Source/Programs/GitDependencies/Program.cs b/Engine/Source/Programs/GitDependencies/Program.cs
new file mode 100644
--- /dev/null
+++ b/Engine/Source/Programs/GitDependencies/Program.cs
@@ -0,0 +1,681 @@
+﻿using System;
+using System.Collections.Concurrent;
+using System.Collections.Generic;
+using System.ComponentModel;
+using System.Diagnostics;
+using System.IO;
+using System.IO.Compression;
+using System.Linq;
+using System.Net;
+using System.Reflection;
+using System.Security.Cryptography;
+using System.Text;
+using System.Threading;
+using System.Threading.Tasks;
+using System.Xml.Serialization;
+
+namespace GitDependencies
+{
+	class Program
+	{
+		class AsyncDownloadState
+		{
+			public int NumFilesRead;
+			public long NumBytesRead;
+			public string ErrorMessage;
+		}
+
+		const string IncomingFileSuffix = ".incoming";
+		const string TempManifestExtension = ".tmp";
+
+		static int Main(string[] Args)
+		{
+			// Build the argument list. Remove any double-hyphens from the start of arguments for conformity with other Epic tools.
+			List<string> ArgsList = new List<string>();
+			foreach (string Arg in Args)
+			{
+				ArgsList.Add(Arg.StartsWith("--")? Arg.Substring(1) : Arg);
+			}
+
+			// Parse the parameters
+			bool bForce = ParseSwitch(ArgsList, "-force");
+			int NumThreads = int.Parse(ParseParameter(ArgsList, "-threads=", "4"));
+			bool bDryRun = ParseSwitch(ArgsList, "-dry-run");
+			bool bHelp = ParseSwitch(ArgsList, "-help");
+			string RootPath = ParseParameter(ArgsList, "-root=", Path.GetFullPath(Path.Combine(Path.GetDirectoryName(Assembly.GetExecutingAssembly().Location), "../../..")));
+			string ManifestPath = ParseParameter(ArgsList, "-manifest=", "Engine\\Build\\Dependencies.xml");
+
+			// Parse all the default exclude filters
+			HashSet<string> ExcludeFolders = new HashSet<string>(StringComparer.CurrentCultureIgnoreCase);
+			if(!ParseSwitch(ArgsList, "-all"))
+			{
+				if(Environment.OSVersion.Platform != PlatformID.Win32NT)
+				{
+					ExcludeFolders.Add("Win32");
+					ExcludeFolders.Add("Win64");
+				}
+				if(Environment.OSVersion.Platform != PlatformID.MacOSX)
+				{
+					ExcludeFolders.Add("Mac");
+				}
+				if(Environment.GetEnvironmentVariable("EMSCRIPTEN") == null)
+				{
+					ExcludeFolders.Add("HTML5");
+				}
+				if(Environment.GetEnvironmentVariable("NDKROOT") == null)
+				{
+					ExcludeFolders.Add("Android");
+				}
+			}
+
+			// Parse all the explicit include filters
+			foreach(string IncludeFolder in ParseParameters(ArgsList, "-include="))
+			{
+				ExcludeFolders.Remove(IncludeFolder.Replace('\\', '/').Trim('/'));
+			}
+
+			// Parse all the explicit exclude filters
+			foreach(string ExcludeFolder in ParseParameters(ArgsList, "-exclude="))
+			{
+				ExcludeFolders.Add(ExcludeFolder.Replace('\\', '/').Trim('/'));
+			}
+
+			// If there are any more parameters, print an error
+			foreach(string RemainingArg in ArgsList)
+			{
+				Log.WriteLine("Invalid command line parameter: {0}", RemainingArg);
+				Log.WriteLine();
+				bHelp = true;
+			}
+
+			// Print the help message
+			if(bHelp)
+			{
+				Log.WriteLine("Usage:");
+				Log.WriteLine("   GitDependencies [options]");
+				Log.WriteLine();
+				Log.WriteLine("Options:");
+				Log.WriteLine("   -all              Sync all folders");
+				Log.WriteLine("   -include=<X>      Include binaries in folders called <X>");
+				Log.WriteLine("   -exclude=<X>      Exclude binaries in folders called <X>");
+				Log.WriteLine("   -force            Overwrite modified dependency files in the workspace");
+				Log.WriteLine("   -root=<PATH>      Specifies the path to the directory to sync with");
+				Log.WriteLine("   -manifest=<PATH>  Sets the path to the dependency manifest");
+				Log.WriteLine("   -threads=X        Use X threads when downloading new files");
+				Log.WriteLine("   -dry-run          Print a list of outdated files, but don't do anything");
+				if(ExcludeFolders.Count > 0)
+				{
+					Log.WriteLine();
+					Log.WriteLine("Current excluded folders: {0}", String.Join(", ", ExcludeFolders));
+				}
+				return 0;
+			}
+
+			// Update the tree
+			if(!UpdateWorkingTree(bForce, bDryRun, RootPath, Path.Combine(RootPath, ManifestPath), ExcludeFolders, NumThreads))
+			{
+				return 1;
+			}
+
+			return 0;
+		}
+
+		static bool ParseSwitch(List<string> ArgsList, string Name)
+		{
+			for(int Idx = 0; Idx < ArgsList.Count; Idx++)
+			{
+				if(String.Compare(ArgsList[Idx], Name, true) == 0)
+				{
+					ArgsList.RemoveAt(Idx);
+					return true;
+				}
+			}
+			return false;
+		}
+
+		static string ParseParameter(List<string> ArgsList, string Prefix, string Default)
+		{
+			string Value = Default;
+			for(int Idx = 0; Idx < ArgsList.Count; Idx++)
+			{
+				if(ArgsList[Idx].StartsWith(Prefix, StringComparison.CurrentCultureIgnoreCase))
+				{
+					Value = ArgsList[Idx].Substring(Prefix.Length);
+					ArgsList.RemoveAt(Idx);
+					break;
+				}
+			}
+			return Value;
+		}
+
+		static IEnumerable<string> ParseParameters(List<string> ArgsList, string Prefix)
+		{
+			for(;;)
+			{
+				string Value = ParseParameter(ArgsList, Prefix, null);
+				if(Value == null)
+				{
+					break;
+				}
+				yield return Value;
+			}
+		}
+
+		static bool UpdateWorkingTree(bool bForce, bool bDryRun, string RootPath, string ManifestPath, HashSet<string> ExcludeFolders, int NumDownloadThreads)
+		{
+			// Start scanning on the working directory 
+			if(ExcludeFolders.Count > 0)
+			{
+				Log.WriteLine("Checking binary dependencies (excluding {0})...", String.Join(", ", ExcludeFolders));
+			}
+			else
+			{
+				Log.WriteLine("Checking binary dependencies...");
+			}
+
+			// Figure out the path to the working manifest
+			string WorkingManifestPath = Path.Combine(RootPath, ".ue4dependencies");
+
+			// Recover from any interrupted transaction to the working manifest, by moving the temporary file into place.
+			string TempWorkingManifestPath = WorkingManifestPath + TempManifestExtension;
+			if(File.Exists(TempWorkingManifestPath) && !File.Exists(WorkingManifestPath) && !SafeMoveFile(TempWorkingManifestPath, WorkingManifestPath))
+			{
+				return false;
+			}
+
+			// Read the initial manifest, or create a new one
+			WorkingManifest CurrentManifest;
+			if(File.Exists(WorkingManifestPath))
+			{
+				CurrentManifest = ReadXmlObject<WorkingManifest>(WorkingManifestPath);
+			}
+			else
+			{
+				CurrentManifest = new WorkingManifest();
+			}
+
+			// Remove all the in-progress download files left over from previous runs
+			foreach(WorkingFile InitialFile in CurrentManifest.Files)
+			{
+				if(InitialFile.Timestamp == 0)
+				{
+					string IncomingFilePath = Path.Combine(RootPath, InitialFile.Name + IncomingFileSuffix);
+					if(File.Exists(IncomingFilePath) && !SafeDeleteFile(IncomingFilePath))
+					{
+						return false;
+					}
+				}
+			}
+
+			// Find all the existing files in the working directory from previous runs. Use the working manifest to cache hashes for them based on timestamp, but recalculate them as needed.
+			Dictionary<string, WorkingFile> CurrentFileLookup = new Dictionary<string, WorkingFile>();
+			foreach(WorkingFile CurrentFile in CurrentManifest.Files)
+			{
+				// Update the hash for this file
+				string CurrentFilePath = Path.Combine(RootPath, CurrentFile.Name);
+				if(File.Exists(CurrentFilePath))
+				{
+					long LastWriteTime = File.GetLastWriteTimeUtc(CurrentFilePath).Ticks;
+					if(LastWriteTime != CurrentFile.Timestamp)
+					{
+						CurrentFile.Hash = ComputeHashForFile(CurrentFilePath);
+						CurrentFile.Timestamp = LastWriteTime;
+					}
+					CurrentFileLookup.Add(CurrentFile.Name, CurrentFile);
+				}
+			}
+
+			// Read the new manifest
+			DependencyManifest NewTargetManifest = ReadXmlObject<DependencyManifest>(ManifestPath);
+			if(NewTargetManifest == null)
+			{
+				return false;
+			}
+
+			// Create a list of files which need to be updated
+			List<DependencyFile> FilesToDownload = new List<DependencyFile>();
+
+			// Create a new working manifest for the working directory, moving over files that we already have. Add any missing dependencies into the download queue.
+			WorkingManifest NewWorkingManifest = new WorkingManifest();
+			foreach(DependencyFile TargetFile in NewTargetManifest.Files)
+			{
+				if(!IsExcludedFolder(TargetFile.Name, ExcludeFolders))
+				{
+					WorkingFile NewFile;
+					if(CurrentFileLookup.TryGetValue(TargetFile.Name, out NewFile) && NewFile.Hash == TargetFile.Hash)
+					{
+						// Update the expected hash to match what we're looking for
+						NewFile.ExpectedHash = TargetFile.Hash;
+
+						// Move the existing file to the new working set
+						CurrentFileLookup.Remove(NewFile.Name);
+					}
+					else
+					{
+						// Create a new working file
+						NewFile = new WorkingFile();
+						NewFile.Name = TargetFile.Name;
+						NewFile.ExpectedHash = TargetFile.Hash;
+
+						// Add it to the download list
+						FilesToDownload.Add(TargetFile);
+					}
+					NewWorkingManifest.Files.Add(NewFile);
+				}
+			}
+
+			// Print out everything that we'd change in a dry run
+			if(bDryRun)
+			{
+				HashSet<string> NewFiles = new HashSet<string>(FilesToDownload.Select(x => x.Name));
+				foreach(string RemoveFile in CurrentFileLookup.Keys.Where(x => !NewFiles.Contains(x)))
+				{
+					Log.WriteLine("Remove {0}", RemoveFile);
+				}
+				foreach(string UpdateFile in CurrentFileLookup.Keys.Where(x => NewFiles.Contains(x)))
+				{
+					Log.WriteLine("Update {0}", UpdateFile);
+				}
+				foreach(string AddFile in NewFiles.Where(x => !CurrentFileLookup.ContainsKey(x)))
+				{
+					Log.WriteLine("Add {0}", AddFile);
+				}
+				return true;
+			}
+
+			// Check if there are any files that have been tampered with
+			if(CurrentFileLookup.Count > 0 && !bForce)
+			{
+				WorkingFile[] ModifiedFiles = CurrentFileLookup.Values.Where(x => x.Hash != x.ExpectedHash).ToArray();
+				if(ModifiedFiles.Length > 0)
+				{
+					Log.WriteError("The following file(s) have been modified. Re-run with the --force parameter to overwrite them.");
+					foreach(WorkingFile ModifiedFile in ModifiedFiles)
+					{
+						Log.WriteError("  {0}", ModifiedFile.Name);
+					}
+					return false;
+				}
+			}
+
+			// Delete any files that are no longer needed
+			foreach(string FileToRemove in CurrentFileLookup.Keys)
+			{
+				if(!SafeDeleteFile(Path.Combine(RootPath, FileToRemove)))
+				{
+					return false;
+				}
+			}
+
+			// Write out the new working manifest, so we can track any files that we're going to download. We always verify missing files on startup, so it's ok that things don't exist yet.
+			if(!WriteWorkingManifest(WorkingManifestPath, TempWorkingManifestPath, NewWorkingManifest))
+			{
+				return false;
+			}
+
+			// Log out if we have any up-to-date files
+			if(NewWorkingManifest.Files.Count > FilesToDownload.Count)
+			{
+				Log.WriteLine("{0} files up to date.", NewWorkingManifest.Files.Count - FilesToDownload.Count);
+			}
+
+			// If there's nothing to do, just print a simpler message and exit early
+			if(FilesToDownload.Count > 0)
+			{
+				// Download all the new dependencies
+				if(!DownloadDependencies(RootPath, FilesToDownload, NewTargetManifest, NumDownloadThreads))
+				{
+					return false;
+				}
+
+				// Update all the timestamps and hashes for the output files
+				foreach(WorkingFile NewFile in NewWorkingManifest.Files)
+				{
+					if(NewFile.Hash != NewFile.ExpectedHash)
+					{
+						string NewFileName = Path.Combine(RootPath, NewFile.Name);
+						NewFile.Hash = NewFile.ExpectedHash;
+						NewFile.Timestamp = File.GetLastWriteTimeUtc(NewFileName).Ticks;
+					}
+				}
+
+				// Rewrite the manifest with the results
+				if(!WriteWorkingManifest(WorkingManifestPath, TempWorkingManifestPath, NewWorkingManifest))
+				{
+					return false;
+				}
+			}
+			return true;
+		}
+
+		static bool IsExcludedFolder(string Name, IEnumerable<string> ExcludeFolders)
+		{
+			foreach(string ExcludeFolder in ExcludeFolders)
+			{
+				if(Name.IndexOf("/" + ExcludeFolder + "/", StringComparison.CurrentCultureIgnoreCase) != -1)
+				{
+					return true;
+				}
+			}
+			return false;
+		}
+
+		static bool DownloadDependencies(string RootPath, IEnumerable<DependencyFile> RequiredFiles, DependencyManifest Manifest, int NumThreads)
+		{
+			// Build a lookup for the files that need updating from each blob
+			Dictionary<string, List<DependencyFile>> BlobToFiles = new Dictionary<string,List<DependencyFile>>();
+			foreach(DependencyFile RequiredFile in RequiredFiles)
+			{
+				List<DependencyFile> FileList;
+				if(!BlobToFiles.TryGetValue(RequiredFile.Hash, out FileList))
+				{
+					FileList = new List<DependencyFile>();
+					BlobToFiles.Add(RequiredFile.Hash, FileList);
+				}
+				FileList.Add(RequiredFile);
+			}
+
+			// Find all the required blobs
+			DependencyBlob[] RequiredBlobs = Manifest.Blobs.Where(x => BlobToFiles.ContainsKey(x.Hash)).ToArray();
+
+			// Build a lookup for the files that need updating from each blob
+			Dictionary<string, List<DependencyBlob>> PackToBlobs = new Dictionary<string,List<DependencyBlob>>();
+			foreach(DependencyBlob RequiredBlob in RequiredBlobs)
+			{
+				List<DependencyBlob> BlobList = new List<DependencyBlob>();
+				if(!PackToBlobs.TryGetValue(RequiredBlob.PackHash, out BlobList))
+				{
+					BlobList = new List<DependencyBlob>();
+					PackToBlobs.Add(RequiredBlob.PackHash, BlobList);
+				}
+				BlobList.Add(RequiredBlob);
+			}
+
+			// Find all the required packs
+			DependencyPack[] RequiredPacks = Manifest.Packs.Where(x => PackToBlobs.ContainsKey(x.Hash)).ToArray();
+
+			// Setup the async state
+			AsyncDownloadState State = new AsyncDownloadState();
+			ConcurrentQueue<DependencyPack> RemainingPacks = new ConcurrentQueue<DependencyPack>(RequiredPacks);
+			long NumBytesTotal = RequiredPacks.Sum(x => x.CompressedSize);
+			long NumFilesTotal = RequiredFiles.Count();
+
+			// Create all the worker threads
+			Thread[] WorkerThreads = new Thread[NumThreads];
+			for(int Idx = 0; Idx < NumThreads; Idx++)
+			{
+				WorkerThreads[Idx] = new Thread(x => DownloadWorker(RootPath, RemainingPacks, PackToBlobs, BlobToFiles, State));
+				WorkerThreads[Idx].Start();
+			}
+
+			// Tick the status message until we've finished or ended with an error
+			while(State.NumFilesRead < NumFilesTotal && State.ErrorMessage == null)
+			{
+				Thread.Sleep(100);
+				long NumBytesRead = Interlocked.Read(ref State.NumBytesRead);
+				Log.WriteStatus("Downloaded {0}/{1} files... ({2:0.0}/{3:0.0}mb; {4}%)", State.NumFilesRead, NumFilesTotal, NumBytesRead / (1024.0 * 1024.0), NumBytesTotal / (1024.0 * 1024.0), (NumBytesRead * 100) / NumBytesTotal);
+			}
+
+			// If we finished with an error
+			if(State.ErrorMessage == null)
+			{
+				Log.FlushStatus();
+				for(int Idx = 0; Idx < WorkerThreads.Length; Idx++)
+				{
+					WorkerThreads[Idx].Join();
+				}
+				return true;
+			}
+			else
+			{
+				Log.WriteError("{0}", State.ErrorMessage);
+				for(int Idx = 0; Idx < WorkerThreads.Length; Idx++)
+				{
+					WorkerThreads[Idx].Abort();
+				}
+				return false;
+			}
+		}
+
+		static void DownloadWorker(string RootPath, ConcurrentQueue<DependencyPack> RemainingPacks, Dictionary<string, List<DependencyBlob>> PackToBlobs, Dictionary<string, List<DependencyFile>> BlobToFiles, AsyncDownloadState State)
+		{
+			DependencyPack NextPack;
+			while(RemainingPacks.TryDequeue(out NextPack))
+			{
+				// Get the temporary filename for the pack file
+				string PackFileName = Path.GetTempFileName();
+				try
+				{
+					// Format the URL for it
+					string BundleUrl = String.Format("http://s3.amazonaws.com/unrealengine/dependencies/{0}/{1}", NextPack.RemotePath, NextPack.Hash);
+
+					// Download the file, decompressing and hashing it as we go
+					try
+					{
+						DownloadFileAndVerifyHash(BundleUrl, PackFileName, NextPack.Hash, Size => Interlocked.Add(ref State.NumBytesRead, Size));
+					}
+					catch(Exception Ex)
+					{
+						throw new Exception(String.Format("Failed to download '{0}' to '{1}': {2}", BundleUrl, PackFileName, Ex.Message));
+					}
+
+					// Move the files into their final place
+					foreach(DependencyBlob Blob in PackToBlobs[NextPack.Hash])
+					{
+						foreach(DependencyFile File in BlobToFiles[Blob.Hash])
+						{
+							string OutputFileName = Path.Combine(RootPath, File.Name);
+							try
+							{
+								ExtractBlob(PackFileName, Blob, OutputFileName);
+							}
+							catch(Exception Ex)
+							{
+								throw new Exception(String.Format("Failed to extract '{0}': {1}", OutputFileName, Ex.Message));
+							}
+							Interlocked.Increment(ref State.NumFilesRead);
+						}
+					}
+				}
+				catch(Exception Ex)
+				{
+					Interlocked.CompareExchange(ref State.ErrorMessage, Ex.Message, null);
+					break;
+				}
+				finally
+				{
+					File.Delete(PackFileName);
+				}
+			}
+		}
+
+		static void DownloadFileAndVerifyHash(string Url, string PackFileName, string ExpectedHash, NotifyReadDelegate NotifyRead)
+		{
+			// Create the output folder
+			Directory.CreateDirectory(Path.GetDirectoryName(PackFileName));
+
+			// Download the file, decompressing and hashing it as we go
+			SHA1Managed Hasher = new SHA1Managed();
+			using(FileStream OutputStream = File.OpenWrite(PackFileName))
+			{
+				CryptoStream HashOutputStream = new CryptoStream(OutputStream, Hasher, CryptoStreamMode.Write);
+				using (WebClient Client = new WebClient())
+				{
+					using(NotifyReadStream InputStream = new NotifyReadStream(Client.OpenRead(Url), NotifyRead))
+					{
+						GZipStream DecompressedStream = new GZipStream(InputStream, CompressionMode.Decompress, true);
+						DecompressedStream.CopyTo(HashOutputStream);
+					}
+				}
+				HashOutputStream.FlushFinalBlock();
+			}
+
+			// Check the hash was what we expected
+			string Hash = BitConverter.ToString(Hasher.Hash).ToLower().Replace("-", "");
+			if(Hash != ExpectedHash)
+			{
+				throw new InvalidDataException("Incorrect hash value");
+			}
+		}
+
+		static void ExtractBlob(string PackFileName, DependencyBlob Blob, string OutputFileName)
+		{
+			// Create the output folder
+			Directory.CreateDirectory(Path.GetDirectoryName(OutputFileName));
+
+			// Copy the data to the output file
+			SHA1Managed Hasher = new SHA1Managed();
+			using(FileStream OutputStream = File.OpenWrite(OutputFileName + IncomingFileSuffix))
+			{
+				CryptoStream HashOutputStream = new CryptoStream(OutputStream, Hasher, CryptoStreamMode.Write);
+				using(FileStream InputStream = File.OpenRead(PackFileName))
+				{
+					// Seek to the right position
+					InputStream.Seek(Blob.PackOffset, SeekOrigin.Begin);
+
+					// Extract the data
+					byte[] Buffer = new byte[16384];
+					for(long RemainingSize = Blob.Size; RemainingSize > 0; )
+					{
+						int ReadSize = InputStream.Read(Buffer, 0, (int)Math.Min(RemainingSize, (long)Buffer.Length));
+						if(ReadSize == 0) throw new InvalidDataException("Unexpected end of file");
+						HashOutputStream.Write(Buffer, 0, ReadSize);
+						RemainingSize -= ReadSize;
+					}
+				}
+				HashOutputStream.FlushFinalBlock();
+			}
+
+			// Check the hash was what we expected
+			string Hash = BitConverter.ToString(Hasher.Hash).ToLower().Replace("-", "");
+			if(Hash != Blob.Hash)
+			{
+				throw new InvalidDataException("Incorrect hash value");
+			}
+
+			// Move the file to its final position
+			File.Move(OutputFileName + IncomingFileSuffix, OutputFileName);
+		}
+
+		public static T ReadXmlObject<T>(string FileName)
+		{
+			try
+			{
+				XmlSerializer Serializer = new XmlSerializer(typeof(T));
+				using(StreamReader Reader = new StreamReader(FileName))
+				{
+					return (T)Serializer.Deserialize(Reader);
+				}
+			}
+			catch(Exception Ex)
+			{
+				Log.WriteError("Failed to read '{0}': {1}", FileName, Ex.Message);
+				return default(T);
+			}
+		}
+
+		public static bool WriteXmlObject<T>(string FileName, T XmlObject)
+		{
+			try
+			{
+				XmlSerializer Serializer = new XmlSerializer(typeof(T));
+				using(StreamWriter Writer = new StreamWriter(FileName))
+				{
+					Serializer.Serialize(Writer, XmlObject);
+				}
+				return true;
+			}
+			catch(Exception Ex)
+			{
+				Log.WriteError("Failed to write file '{0}': {1}", FileName, Ex.Message);
+				return false;
+			}
+		}
+
+		public static bool WriteWorkingManifest(string FileName, string TemporaryFileName, WorkingManifest Manifest)
+		{
+			if(!WriteXmlObject(TemporaryFileName, Manifest))
+			{
+				return false;
+			}
+			if(!SafeModifyFileAttributes(TemporaryFileName, FileAttributes.Hidden, 0))
+			{
+				return false;
+			}
+			if(!SafeDeleteFile(FileName))
+			{
+				return false;
+			}
+			if(!SafeMoveFile(TemporaryFileName, FileName))
+			{
+				return false;
+			}
+			return true;
+		}
+
+		public static bool SafeModifyFileAttributes(string FileName, FileAttributes AddAttributes, FileAttributes RemoveAttributes)
+		{
+			try
+			{
+				File.SetAttributes(FileName, (File.GetAttributes(FileName) | AddAttributes) & ~RemoveAttributes);
+				return true;
+			}
+			catch(IOException)
+			{
+				Log.WriteError("Failed to set attributes for file '{0}'", FileName);
+				return false;
+			}
+		}
+
+		public static bool SafeCreateDirectory(string DirectoryName)
+		{
+			try
+			{
+				Directory.CreateDirectory(DirectoryName);
+				return true;
+			}
+			catch(IOException)
+			{
+				Log.WriteError("Failed to create directory '{0}'", DirectoryName);
+				return false;
+			}
+		}
+
+		public static bool SafeDeleteFile(string FileName)
+		{
+			try
+			{
+				File.Delete(FileName);
+				return true;
+			}
+			catch(IOException)
+			{
+				Log.WriteError("Failed to delete file '{0}'", FileName);
+				return false;
+			}
+		}
+
+		public static bool SafeMoveFile(string SourceFileName, string TargetFileName)
+		{
+			try
+			{
+				File.Move(SourceFileName, TargetFileName);
+				return true;
+			}
+			catch(IOException)
+			{
+				Log.WriteError("Failed to rename '{0}' to '{1}'", SourceFileName, TargetFileName);
+				return false;
+			}
+		}
+
+		public static string ComputeHashForFile(string FileName)
+		{
+			using(FileStream InputStream = File.OpenRead(FileName))
+			{
+				byte[] Hash = new SHA1CryptoServiceProvider().ComputeHash(InputStream);
+				return BitConverter.ToString(Hash).ToLower().Replace("-", "");
+			}
+		}
+	}
+}
diff --git a/Engine/Source/Programs/GitDependencies/Properties/AssemblyInfo.cs b/Engine/Source/Programs/GitDependencies/Properties/AssemblyInfo.cs
new file mode 100644
--- /dev/null
+++ b/Engine/Source/Programs/GitDependencies/Properties/AssemblyInfo.cs
@@ -0,0 +1,10 @@
+﻿using System.Reflection;
+using System.Runtime.CompilerServices;
+using System.Runtime.InteropServices;
+
+// General Information about an assembly is controlled through the following 
+// set of attributes. Change these attribute values to modify the information
+// associated with an assembly.
+[assembly: AssemblyTitle("GitDependencies")]
+[assembly: AssemblyDescription("")]
+[assembly: AssemblyConfiguration("")]
diff --git a/Engine/Source/Programs/GitDependencies/WorkingManifest.cs b/Engine/Source/Programs/GitDependencies/WorkingManifest.cs
new file mode 100644
--- /dev/null
+++ b/Engine/Source/Programs/GitDependencies/WorkingManifest.cs
@@ -0,0 +1,30 @@
+﻿using System;
+using System.Collections.Generic;
+using System.Linq;
+using System.Text;
+using System.Threading.Tasks;
+using System.Xml.Serialization;
+
+namespace GitDependencies
+{
+	public class WorkingFile
+	{
+		[XmlAttribute]
+		public string Name;
+
+		[XmlAttribute]
+		public string Hash;
+
+		[XmlAttribute]
+		public string ExpectedHash;
+
+		[XmlAttribute]
+		public long Timestamp;
+	}
+
+	public class WorkingManifest
+	{
+		[XmlArrayItem("File")]
+		public List<WorkingFile> Files = new List<WorkingFile>();
+	}
+}
diff --git a/Engine/Source/Programs/UnrealBuildTool/System/ProjectFileGenerator.cs b/Engine/Source/Programs/UnrealBuildTool/System/ProjectFileGenerator.cs
--- a/Engine/Source/Programs/UnrealBuildTool/System/ProjectFileGenerator.cs
+++ b/Engine/Source/Programs/UnrealBuildTool/System/ProjectFileGenerator.cs
@@ -493,6 +493,9 @@ namespace UnrealBuildTool
 					// Add DotNETUtilities to the master project
 					ProgramsFolder.ChildProjects.Add(AddSimpleCSharpProject("DotNETCommon/DotNETUtilities", bShouldBuildForAllSolutionTargets: true, bForceDevelopmentConfiguration: true));
 
+					// Add the Git dependencies project
+					ProgramsFolder.ChildProjects.Add(AddSimpleCSharpProject("GitDependencies", bShouldBuildForAllSolutionTargets: true, bForceDevelopmentConfiguration: true));
+
 					// Add all of the IOS C# projects
 					AddIOSProjects( ProgramsFolder );
 
