Support for retrying when a download fails, and aggregating multiple dependency files in the Build directory.

diff --git a/Engine/Source/Programs/GitDependencies/Program.cs b/Engine/Source/Programs/GitDependencies/Program.cs
--- a/Engine/Source/Programs/GitDependencies/Program.cs
+++ b/Engine/Source/Programs/GitDependencies/Program.cs
@@ -40,10 +40,10 @@ namespace GitDependencies
 			// Parse the parameters
 			bool bForce = ParseSwitch(ArgsList, "-force");
 			int NumThreads = int.Parse(ParseParameter(ArgsList, "-threads=", "4"));
+			int MaxRetries = int.Parse(ParseParameter(ArgsList, "-max-retries=", "4"));
 			bool bDryRun = ParseSwitch(ArgsList, "-dry-run");
 			bool bHelp = ParseSwitch(ArgsList, "-help");
 			string RootPath = ParseParameter(ArgsList, "-root=", Path.GetFullPath(Path.Combine(Path.GetDirectoryName(Assembly.GetExecutingAssembly().Location), "../../..")));
-			string ManifestPath = ParseParameter(ArgsList, "-manifest=", "Engine\\Build\\Dependencies.xml");
 
 			// Parse all the default exclude filters
 			HashSet<string> ExcludeFolders = new HashSet<string>(StringComparer.CurrentCultureIgnoreCase);
@@ -100,9 +100,9 @@ namespace GitDependencies
 				Log.WriteLine("   -exclude=<X>      Exclude binaries in folders called <X>");
 				Log.WriteLine("   -force            Overwrite modified dependency files in the workspace");
 				Log.WriteLine("   -root=<PATH>      Specifies the path to the directory to sync with");
-				Log.WriteLine("   -manifest=<PATH>  Sets the path to the dependency manifest");
 				Log.WriteLine("   -threads=X        Use X threads when downloading new files");
 				Log.WriteLine("   -dry-run          Print a list of outdated files, but don't do anything");
+				Log.WriteLine("   -max-retries		Set the maximum number of retries for downloading files");
 				if(ExcludeFolders.Count > 0)
 				{
 					Log.WriteLine();
@@ -112,7 +112,7 @@ namespace GitDependencies
 			}
 
 			// Update the tree
-			if(!UpdateWorkingTree(bForce, bDryRun, RootPath, Path.Combine(RootPath, ManifestPath), ExcludeFolders, NumThreads))
+			if(!UpdateWorkingTree(bForce, bDryRun, RootPath, ExcludeFolders, NumThreads, MaxRetries))
 			{
 				return 1;
 			}
@@ -161,7 +161,7 @@ namespace GitDependencies
 			}
 		}
 
-		static bool UpdateWorkingTree(bool bForce, bool bDryRun, string RootPath, string ManifestPath, HashSet<string> ExcludeFolders, int NumDownloadThreads)
+		static bool UpdateWorkingTree(bool bForce, bool bDryRun, string RootPath, HashSet<string> ExcludeFolders, int NumThreads, int MaxRetries)
 		{
 			// Start scanning on the working directory 
 			if(ExcludeFolders.Count > 0)
@@ -207,11 +207,39 @@ namespace GitDependencies
 				}
 			}
 
-			// Read the new manifest
-			DependencyManifest NewTargetManifest = ReadXmlObject<DependencyManifest>(ManifestPath);
-			if(NewTargetManifest == null)
+			// Find all the manifests and push them into dictionaries
+			Dictionary<string, DependencyFile> TargetFiles = new Dictionary<string,DependencyFile>(StringComparer.InvariantCultureIgnoreCase);
+			Dictionary<string, DependencyBlob> TargetBlobs = new Dictionary<string,DependencyBlob>(StringComparer.InvariantCultureIgnoreCase);
+			Dictionary<string, DependencyPack> TargetPacks = new Dictionary<string,DependencyPack>(StringComparer.InvariantCultureIgnoreCase);
+			foreach(string BaseFolder in Directory.EnumerateDirectories(RootPath))
 			{
-				return false;
+				string BuildFolder = Path.Combine(BaseFolder, "Build");
+				if(Directory.Exists(BuildFolder))
+				{
+					foreach(string ManifestFileName in Directory.EnumerateFiles(BuildFolder, "*.gitdeps.xml"))
+					{
+						// Read this manifest
+						DependencyManifest NewTargetManifest = ReadXmlObject<DependencyManifest>(ManifestFileName);
+						if(NewTargetManifest == null)
+						{
+							return false;
+						}
+
+						// Add all the files, blobs and packs into the shared dictionaries
+						foreach(DependencyFile NewFile in NewTargetManifest.Files)
+						{
+							TargetFiles[NewFile.Name] = NewFile;
+						}
+						foreach(DependencyBlob NewBlob in NewTargetManifest.Blobs)
+						{
+							TargetBlobs[NewBlob.Hash] = NewBlob;
+						}
+						foreach(DependencyPack NewPack in NewTargetManifest.Packs)
+						{
+							TargetPacks[NewPack.Hash] = NewPack;
+						}
+					}
+				}
 			}
 
 			// Find all the existing files in the working directory from previous runs. Use the working manifest to cache hashes for them based on timestamp, but recalculate them as needed.
@@ -237,7 +265,7 @@ namespace GitDependencies
 
 			// Create a new working manifest for the working directory, moving over files that we already have. Add any missing dependencies into the download queue.
 			WorkingManifest NewWorkingManifest = new WorkingManifest();
-			foreach(DependencyFile TargetFile in NewTargetManifest.Files)
+			foreach(DependencyFile TargetFile in TargetFiles.Values)
 			{
 				if(!IsExcludedFolder(TargetFile.Name, ExcludeFolders))
 				{
@@ -322,8 +350,10 @@ namespace GitDependencies
 			// If there's nothing to do, just print a simpler message and exit early
 			if(FilesToDownload.Count > 0)
 			{
+				Log.WriteLine("Downloading {0} files using {1} threads...", FilesToDownload.Count, NumThreads);
+
 				// Download all the new dependencies
-				if(!DownloadDependencies(RootPath, FilesToDownload, NewTargetManifest, NumDownloadThreads))
+				if(!DownloadDependencies(RootPath, FilesToDownload, TargetBlobs.Values, TargetPacks.Values, NumThreads, MaxRetries))
 				{
 					return false;
 				}
@@ -360,7 +390,7 @@ namespace GitDependencies
 			return false;
 		}
 
-		static bool DownloadDependencies(string RootPath, IEnumerable<DependencyFile> RequiredFiles, DependencyManifest Manifest, int NumThreads)
+		static bool DownloadDependencies(string RootPath, IEnumerable<DependencyFile> RequiredFiles, IEnumerable<DependencyBlob> Blobs, IEnumerable<DependencyPack> Packs, int NumThreads, int MaxRetries)
 		{
 			// Build a lookup for the files that need updating from each blob
 			Dictionary<string, List<DependencyFile>> BlobToFiles = new Dictionary<string,List<DependencyFile>>();
@@ -376,7 +406,7 @@ namespace GitDependencies
 			}
 
 			// Find all the required blobs
-			DependencyBlob[] RequiredBlobs = Manifest.Blobs.Where(x => BlobToFiles.ContainsKey(x.Hash)).ToArray();
+			DependencyBlob[] RequiredBlobs = Blobs.Where(x => BlobToFiles.ContainsKey(x.Hash)).ToArray();
 
 			// Build a lookup for the files that need updating from each blob
 			Dictionary<string, List<DependencyBlob>> PackToBlobs = new Dictionary<string,List<DependencyBlob>>();
@@ -392,7 +422,7 @@ namespace GitDependencies
 			}
 
 			// Find all the required packs
-			DependencyPack[] RequiredPacks = Manifest.Packs.Where(x => PackToBlobs.ContainsKey(x.Hash)).ToArray();
+			DependencyPack[] RequiredPacks = Packs.Where(x => PackToBlobs.ContainsKey(x.Hash)).ToArray();
 
 			// Setup the async state
 			AsyncDownloadState State = new AsyncDownloadState();
@@ -404,7 +434,7 @@ namespace GitDependencies
 			Thread[] WorkerThreads = new Thread[NumThreads];
 			for(int Idx = 0; Idx < NumThreads; Idx++)
 			{
-				WorkerThreads[Idx] = new Thread(x => DownloadWorker(RootPath, RemainingPacks, PackToBlobs, BlobToFiles, State));
+				WorkerThreads[Idx] = new Thread(x => DownloadWorker(RootPath, RemainingPacks, PackToBlobs, BlobToFiles, State, MaxRetries));
 				WorkerThreads[Idx].Start();
 			}
 
@@ -413,7 +443,7 @@ namespace GitDependencies
 			{
 				Thread.Sleep(100);
 				long NumBytesRead = Interlocked.Read(ref State.NumBytesRead);
-				Log.WriteStatus("Downloaded {0}/{1} files... ({2:0.0}/{3:0.0}mb; {4}%)", State.NumFilesRead, NumFilesTotal, NumBytesRead / (1024.0 * 1024.0), NumBytesTotal / (1024.0 * 1024.0), (NumBytesRead * 100) / NumBytesTotal);
+				Log.WriteStatus("Staged {0}/{1} files ({2:0.0}/{3:0.0}mb; {4}%)", State.NumFilesRead, NumFilesTotal, NumBytesRead / (1024.0 * 1024.0), NumBytesTotal / (1024.0 * 1024.0), (NumBytesRead * 100) / NumBytesTotal);
 			}
 
 			// If we finished with an error
@@ -437,7 +467,7 @@ namespace GitDependencies
 			}
 		}
 
-		static void DownloadWorker(string RootPath, ConcurrentQueue<DependencyPack> RemainingPacks, Dictionary<string, List<DependencyBlob>> PackToBlobs, Dictionary<string, List<DependencyFile>> BlobToFiles, AsyncDownloadState State)
+		static void DownloadWorker(string RootPath, ConcurrentQueue<DependencyPack> RemainingPacks, Dictionary<string, List<DependencyBlob>> PackToBlobs, Dictionary<string, List<DependencyFile>> BlobToFiles, AsyncDownloadState State, int MaxRetries)
 		{
 			DependencyPack NextPack;
 			while(RemainingPacks.TryDequeue(out NextPack))
@@ -449,14 +479,35 @@ namespace GitDependencies
 					// Format the URL for it
 					string BundleUrl = String.Format("http://s3.amazonaws.com/unrealengine/dependencies/{0}/{1}", NextPack.RemotePath, NextPack.Hash);
 
-					// Download the file, decompressing and hashing it as we go
-					try
-					{
-						DownloadFileAndVerifyHash(BundleUrl, PackFileName, NextPack.Hash, Size => Interlocked.Add(ref State.NumBytesRead, Size));
-					}
-					catch(Exception Ex)
+					// Download the file, decompressing and hashing it as we go.
+					for(int NumAttempts = 0;;)
 					{
-						throw new Exception(String.Format("Failed to download '{0}' to '{1}': {2}", BundleUrl, PackFileName, Ex.Message));
+						string ExceptionMessage;
+
+						// Try to download the file
+						long RollbackSize = 0;
+						try
+						{
+							DownloadFileAndVerifyHash(BundleUrl, PackFileName, NextPack.Hash, Size => { RollbackSize += Size; Interlocked.Add(ref State.NumBytesRead, Size); });
+							break;
+						}
+						catch(InvalidDataException Ex)
+						{
+							// The downloaded file was corrupt
+							ExceptionMessage = Ex.Message;
+						}
+						catch(WebException Ex)
+						{
+							// Error while downloading the file
+							ExceptionMessage = Ex.Message;
+						}
+						Interlocked.Add(ref State.NumBytesRead, -RollbackSize);
+
+						// Bail out after retrying
+						if(++NumAttempts == MaxRetries)
+						{
+							throw new Exception(String.Format("Failed to download '{0}' to '{1}' after {2} attempts: {3}", BundleUrl, PackFileName, NumAttempts, ExceptionMessage));
+						}
 					}
 
 					// Move the files into their final place
