Allow GitDependencies to detect Mac correctly, and handle cases where files to download would already exist in the workspace.

diff --git a/Engine/Source/Programs/GitDependencies/Program.cs b/Engine/Source/Programs/GitDependencies/Program.cs
--- a/Engine/Source/Programs/GitDependencies/Program.cs
+++ b/Engine/Source/Programs/GitDependencies/Program.cs
@@ -11,7 +11,6 @@ using System.Reflection;
 using System.Security.Cryptography;
 using System.Text;
 using System.Threading;
-using System.Threading.Tasks;
 using System.Xml.Serialization;
 
 namespace GitDependencies
@@ -54,7 +53,7 @@ namespace GitDependencies
 					ExcludeFolders.Add("Win32");
 					ExcludeFolders.Add("Win64");
 				}
-				if(Environment.OSVersion.Platform != PlatformID.MacOSX)
+				if(Environment.OSVersion.Platform != PlatformID.MacOSX && !(Environment.OSVersion.Platform == PlatformID.Unix && Directory.Exists("/Applications") && Directory.Exists("/System")))
 				{
 					ExcludeFolders.Add("Mac");
 				}
@@ -262,6 +261,23 @@ namespace GitDependencies
 				}
 			}
 
+			// Also add all the untracked files which already exist, but weren't downloaded by this program
+			foreach (DependencyFile TargetFile in TargetFiles.Values) 
+			{
+				if(!CurrentFileLookup.ContainsKey(TargetFile.Name))
+				{
+					string CurrentFilePath = Path.Combine(RootPath, TargetFile.Name);
+					if(File.Exists(CurrentFilePath))
+					{
+						WorkingFile CurrentFile = new WorkingFile();
+						CurrentFile.Name = TargetFile.Name;
+						CurrentFile.Hash = ComputeHashForFile(CurrentFilePath);
+						CurrentFile.Timestamp = File.GetLastWriteTimeUtc(CurrentFilePath).Ticks;
+						CurrentFileLookup.Add(CurrentFile.Name, CurrentFile);
+					}
+				}
+			}
+
 			// Create a list of files which need to be updated
 			List<DependencyFile> FilesToDownload = new List<DependencyFile>();
 
@@ -517,7 +533,8 @@ namespace GitDependencies
 						// Bail out after retrying
 						if(++NumAttempts == MaxRetries)
 						{
-							throw new Exception(String.Format("Failed to download '{0}' to '{1}' after {2} attempts: {3} ({4})", BundleUrl, PackFileName, NumAttempts, CaughtException.Message, CaughtException.GetType().Name));
+							Interlocked.CompareExchange(ref State.ErrorMessage, String.Format("Failed to download '{0}' to '{1}' after {2} attempts: {3} ({4})", BundleUrl, PackFileName, NumAttempts, CaughtException.Message, CaughtException.GetType().Name), null);
+							return;
 						}
 					}
 
@@ -533,7 +550,8 @@ namespace GitDependencies
 							}
 							catch(Exception Ex)
 							{
-								throw new Exception(String.Format("Failed to extract '{0}': {1}", OutputFileName, Ex.Message));
+								Interlocked.CompareExchange(ref State.ErrorMessage, String.Format("Failed to extract '{0}': {1}", OutputFileName, Ex.Message), null);
+								return;
 							}
 							Interlocked.Increment(ref State.NumFilesRead);
 						}
@@ -541,7 +559,7 @@ namespace GitDependencies
 				}
 				catch(Exception Ex)
 				{
-					Interlocked.CompareExchange(ref State.ErrorMessage, Ex.Message, null);
+					Interlocked.CompareExchange(ref State.ErrorMessage, Ex.ToString(), null);
 					break;
 				}
 				finally
