Better handling for corrupted manifest files, and allow continuing even if files have been tampered with.

diff --git a/Engine/Source/Programs/GitDependencies/Program.cs b/Engine/Source/Programs/GitDependencies/Program.cs
--- a/Engine/Source/Programs/GitDependencies/Program.cs
+++ b/Engine/Source/Programs/GitDependencies/Program.cs
@@ -186,11 +186,7 @@ namespace GitDependencies
 
 			// Read the initial manifest, or create a new one
 			WorkingManifest CurrentManifest;
-			if(File.Exists(WorkingManifestPath))
-			{
-				CurrentManifest = ReadXmlObject<WorkingManifest>(WorkingManifestPath);
-			}
-			else
+			if(!File.Exists(WorkingManifestPath) || !ReadXmlObject(WorkingManifestPath, out CurrentManifest))
 			{
 				CurrentManifest = new WorkingManifest();
 			}
@@ -220,8 +216,8 @@ namespace GitDependencies
 					foreach(string ManifestFileName in Directory.EnumerateFiles(BuildFolder, "*.gitdeps.xml"))
 					{
 						// Read this manifest
-						DependencyManifest NewTargetManifest = ReadXmlObject<DependencyManifest>(ManifestFileName);
-						if(NewTargetManifest == null)
+						DependencyManifest NewTargetManifest;
+						if(!ReadXmlObject(ManifestFileName, out NewTargetManifest))
 						{
 							return false;
 						}
@@ -329,28 +325,30 @@ namespace GitDependencies
 				return true;
 			}
 
-			// Check if there are any files that have been tampered with
-			if(CurrentFileLookup.Count > 0 && !bForce)
+			// Delete any files which are no longer needed
+			List<WorkingFile> TamperedFiles = new List<WorkingFile>();
+			foreach(WorkingFile FileToRemove in CurrentFileLookup.Values)
 			{
-				WorkingFile[] ModifiedFiles = CurrentFileLookup.Values.Where(x => x.Hash != x.ExpectedHash).ToArray();
-				if(ModifiedFiles.Length > 0)
+				if(!bForce && FileToRemove.Hash != FileToRemove.ExpectedHash)
+				{
+					TamperedFiles.Add(FileToRemove);
+				}
+				else if(!SafeDeleteFile(Path.Combine(RootPath, FileToRemove.Name)))
 				{
-					Log.WriteError("The following file(s) have been modified. Re-run with the --force parameter to overwrite them.");
-					foreach(WorkingFile ModifiedFile in ModifiedFiles)
-					{
-						Log.WriteError("  {0}", ModifiedFile.Name);
-					}
 					return false;
 				}
 			}
 
-			// Delete any files that are no longer needed
-			foreach(string FileToRemove in CurrentFileLookup.Keys)
+			// Warn if there were any files that have been tampered with
+			if(TamperedFiles.Count > 0)
 			{
-				if(!SafeDeleteFile(Path.Combine(RootPath, FileToRemove)))
+				Log.WriteError("The following file(s) have been modified, and were not updated:");
+				foreach(WorkingFile TamperedFile in TamperedFiles)
 				{
-					return false;
+					Log.WriteError("  {0}", TamperedFile.Name);
+					TargetFiles.Remove(TamperedFile.Name);
 				}
+				Log.WriteError("Re-run with the --force parameter to overwrite them.");
 			}
 
 			// Write out the new working manifest, so we can track any files that we're going to download. We always verify missing files on startup, so it's ok that things don't exist yet.
@@ -637,20 +635,22 @@ namespace GitDependencies
 			File.Move(OutputFileName + IncomingFileSuffix, OutputFileName);
 		}
 
-		public static T ReadXmlObject<T>(string FileName)
+		public static bool ReadXmlObject<T>(string FileName, out T NewObject)
 		{
 			try
 			{
 				XmlSerializer Serializer = new XmlSerializer(typeof(T));
 				using(StreamReader Reader = new StreamReader(FileName))
 				{
-					return (T)Serializer.Deserialize(Reader);
+					NewObject = (T)Serializer.Deserialize(Reader);
 				}
+				return true;
 			}
 			catch(Exception Ex)
 			{
-				Log.WriteError("Failed to read '{0}': {1}", FileName, Ex.Message);
-				return default(T);
+				Log.WriteError("Failed to read '{0}': {1}", FileName, Ex.ToString());
+				NewObject = default(T);
+				return false;
 			}
 		}
 
