Remove files from the download queue if they have been tampered with.

diff --git a/Engine/Source/Programs/GitDependencies/Program.cs b/Engine/Source/Programs/GitDependencies/Program.cs
--- a/Engine/Source/Programs/GitDependencies/Program.cs
+++ b/Engine/Source/Programs/GitDependencies/Program.cs
@@ -339,14 +339,19 @@ namespace GitDependencies
 				}
 			}
 
-			// Warn if there were any files that have been tampered with
+			// Warn if there were any files that have been tampered with, and remove them from the download list
 			if(TamperedFiles.Count > 0)
 			{
 				Log.WriteError("The following file(s) have been modified, and were not updated:");
 				foreach(WorkingFile TamperedFile in TamperedFiles)
 				{
+					DependencyFile TargetFile;
+					if(TargetFiles.TryGetValue(TamperedFile.Name, out TargetFile))
+					{
+						TargetFiles.Remove(TamperedFile.Name);
+						FilesToDownload.Remove(TargetFile);
+					}
 					Log.WriteError("  {0}", TamperedFile.Name);
-					TargetFiles.Remove(TamperedFile.Name);
 				}
 				Log.WriteError("Re-run with the --force parameter to overwrite them.");
 			}
