Attempts to improve performance and robustness of GitDependencies.

- Now uses 4 threads by default
- Decompression is done on a separate thread
- Each download thread exists in a succeed or fail state, and program only fails once no progress is being made by any thread
- Uses HttpWebRequest/HttpWebResponse instead of WebClient

diff --git a/Engine/Source/Programs/GitDependencies/Program.cs b/Engine/Source/Programs/GitDependencies/Program.cs
--- a/Engine/Source/Programs/GitDependencies/Program.cs
+++ b/Engine/Source/Programs/GitDependencies/Program.cs
@@ -19,14 +19,17 @@ namespace GitDependencies
 	{
 		class AsyncDownloadState
 		{
+			public int NumFiles;
 			public int NumFilesRead;
 			public long NumBytesRead;
-			public string ErrorMessage;
+			public string LastDecompressError;
+			public int NumFailingDownloads;
+			public string LastDownloadError;
 		}
 
 		const string IncomingFileSuffix = ".incoming";
 		const string TempManifestExtension = ".tmp";
-
+		
 		static int Main(string[] Args)
 		{
 			// Build the argument list. Remove any double-hyphens from the start of arguments for conformity with other Epic tools.
@@ -411,6 +414,9 @@ namespace GitDependencies
 
 		static bool DownloadDependencies(string RootPath, IEnumerable<DependencyFile> RequiredFiles, IEnumerable<DependencyBlob> Blobs, IEnumerable<DependencyPack> Packs, int NumThreads, int MaxRetries)
 		{
+			// Make sure we can actually open the right number of connections
+			ServicePointManager.DefaultConnectionLimit = NumThreads;
+
 			// Build a lookup for the files that need updating from each blob
 			Dictionary<string, List<DependencyFile>> BlobToFiles = new Dictionary<string,List<DependencyFile>>();
 			foreach(DependencyFile RequiredFile in RequiredFiles)
@@ -443,161 +449,189 @@ namespace GitDependencies
 			// Find all the required packs
 			DependencyPack[] RequiredPacks = Packs.Where(x => PackToBlobs.ContainsKey(x.Hash)).ToArray();
 
+			// Get temporary filenames for all the files we're going to download
+			Dictionary<DependencyPack, string> DownloadFileNames = new Dictionary<DependencyPack,string>();
+			foreach(DependencyPack Pack in RequiredPacks)
+			{
+				DownloadFileNames.Add(Pack, Path.GetTempFileName());
+			}
+
 			// Setup the async state
 			AsyncDownloadState State = new AsyncDownloadState();
-			ConcurrentQueue<DependencyPack> RemainingPacks = new ConcurrentQueue<DependencyPack>(RequiredPacks);
+			State.NumFiles = RequiredFiles.Count();
 			long NumBytesTotal = RequiredPacks.Sum(x => x.CompressedSize);
-			long NumFilesTotal = RequiredFiles.Count();
+			ConcurrentQueue<DependencyPack> DownloadQueue = new ConcurrentQueue<DependencyPack>(RequiredPacks);
+			ConcurrentQueue<DependencyPack> DecompressQueue = new ConcurrentQueue<DependencyPack>();
 
 			// Create all the worker threads
 			Thread[] WorkerThreads = new Thread[NumThreads];
 			for(int Idx = 0; Idx < NumThreads; Idx++)
 			{
-				WorkerThreads[Idx] = new Thread(x => DownloadWorker(RootPath, RemainingPacks, PackToBlobs, BlobToFiles, State, MaxRetries));
+				WorkerThreads[Idx] = new Thread(x => DownloadWorker(RootPath, DownloadQueue, DecompressQueue, DownloadFileNames, PackToBlobs, BlobToFiles, State, MaxRetries));
 				WorkerThreads[Idx].Start();
 			}
 
+			// Create the decompression thread
+			Thread DecompressionThread = new Thread(x => DecompressWorker(RootPath, DecompressQueue, DownloadFileNames, PackToBlobs, BlobToFiles, State));
+			DecompressionThread.Start();
+
 			// Tick the status message until we've finished or ended with an error. Use a circular buffer to average out the speed over time.
-			long[] NumBytesReadBuffer = new long[40];
-			for(int BufferIdx = 0; State.NumFilesRead < NumFilesTotal && State.ErrorMessage == null; BufferIdx = (BufferIdx + 1) % NumBytesReadBuffer.Length)
+			long[] NumBytesReadBuffer = new long[60];
+			for(int BufferIdx = 0; State.NumFilesRead < State.NumFiles && State.NumFailingDownloads < NumThreads && State.LastDecompressError == null; BufferIdx = (BufferIdx + 1) % NumBytesReadBuffer.Length)
 			{
 				const int TickInterval = 100;
 
 				long NumBytesRead = Interlocked.Read(ref State.NumBytesRead);
 				float NumBytesPerSecond = (float)Math.Max(NumBytesRead - NumBytesReadBuffer[BufferIdx], 0) * 1000.0f / (NumBytesReadBuffer.Length * TickInterval);
-				Log.WriteStatus("Staged {0}/{1} files ({2:0.0}/{3:0.0}mb; {4:0.00}mb/s; {5}%)...", State.NumFilesRead, NumFilesTotal, (NumBytesRead / (1024.0 * 1024.0)) + 0.0999999, (NumBytesTotal / (1024.0 * 1024.0)) + 0.0999999, (NumBytesPerSecond / (1024.0 * 1024.0)) + 0.00999999, (NumBytesRead * 100) / NumBytesTotal);
+				Log.WriteStatus("Staged {0}/{1} files ({2:0.0}/{3:0.0}mb; {4:0.00}mb/s; {5}%)...", State.NumFilesRead, State.NumFiles, (NumBytesRead / (1024.0 * 1024.0)) + 0.0999999, (NumBytesTotal / (1024.0 * 1024.0)) + 0.0999999, (NumBytesPerSecond / (1024.0 * 1024.0)) + 0.00999999, (NumBytesRead * 100) / NumBytesTotal);
 				NumBytesReadBuffer[BufferIdx] = NumBytesRead;
 
 				Thread.Sleep(TickInterval);
 			}
 
-			// If we finished with an error
-			if(State.ErrorMessage == null)
+			// If we finished with an error, try to clean up and return
+			if(State.NumFilesRead < State.NumFiles)
 			{
-				Log.FlushStatus();
-				for(int Idx = 0; Idx < WorkerThreads.Length; Idx++)
+				DecompressionThread.Abort();
+				foreach(Thread WorkerThread in WorkerThreads)
 				{
-					WorkerThreads[Idx].Join();
+					WorkerThread.Abort();
 				}
-				return true;
-			}
-			else
-			{
-				Log.WriteError("{0}", State.ErrorMessage);
-				for(int Idx = 0; Idx < WorkerThreads.Length; Idx++)
+				Log.WriteError("{0}", (State.LastDecompressError != null)? State.LastDecompressError : State.LastDownloadError);
+				foreach(string FileName in DownloadFileNames.Values)
 				{
-					WorkerThreads[Idx].Abort();
+					try { File.Delete(FileName); } catch (Exception) { }
 				}
 				return false;
 			}
+
+			// Join all the threads
+			DecompressionThread.Join();
+			foreach(Thread WorkerThread in WorkerThreads)
+			{
+				WorkerThread.Join();
+			}
+			Log.FlushStatus();
+			return true;
 		}
 
-		static void DownloadWorker(string RootPath, ConcurrentQueue<DependencyPack> RemainingPacks, Dictionary<string, List<DependencyBlob>> PackToBlobs, Dictionary<string, List<DependencyFile>> BlobToFiles, AsyncDownloadState State, int MaxRetries)
+		static void DecompressWorker(string RootPath, ConcurrentQueue<DependencyPack> DecompressQueue, Dictionary<DependencyPack, string> DownloadFileNames, Dictionary<string, List<DependencyBlob>> PackToBlobs, Dictionary<string, List<DependencyFile>> BlobToFiles, AsyncDownloadState State)
 		{
-			DependencyPack NextPack;
-			while(RemainingPacks.TryDequeue(out NextPack))
+			while(State.NumFilesRead < State.NumFiles)
 			{
-				// Get the temporary filename for the pack file
-				string PackFileName = Path.GetTempFileName();
-				try
+				// Remove the next file from the queue, or wait before polling again
+				DependencyPack NextPack;
+				if(!DecompressQueue.TryDequeue(out NextPack))
 				{
-					// Format the URL for it
-					string BundleUrl = String.Format("http://cdn.unrealengine.com/dependencies/{0}/{1}", NextPack.RemotePath, NextPack.Hash);
+					Thread.Sleep(100);
+					continue;
+				}
 
-					// Download the file, decompressing and hashing it as we go.
-					for(int NumAttempts = 0;;)
-					{
-						Exception CaughtException;
+				// Get the filename for the downloaded pack
+				string PackFileName = DownloadFileNames[NextPack];
 
-						// Try to download the file
-						long RollbackSize = 0;
+				// Extract all the files from this pack file to their final locations
+				foreach(DependencyBlob Blob in PackToBlobs[NextPack.Hash])
+				{
+					foreach(DependencyFile File in BlobToFiles[Blob.Hash])
+					{
+						string OutputFileName = Path.Combine(RootPath, File.Name);
 						try
 						{
-							DownloadFileAndVerifyHash(BundleUrl, PackFileName, NextPack.Hash, Size => { RollbackSize += Size; Interlocked.Add(ref State.NumBytesRead, Size); });
-							break;
-						}
-						catch(InvalidDataException Ex)
-						{
-							// The downloaded file was corrupt
-							CaughtException = Ex;
-						}
-						catch(IOException Ex)
-						{
-							// Couldn't read/write from the stream or file
-							CaughtException = Ex;
+							ExtractBlob(PackFileName, Blob, OutputFileName);
 						}
-						catch(WebException Ex)
+						catch(Exception Ex)
 						{
-							// Error while downloading the file
-							CaughtException = Ex;
-						}
-						Interlocked.Add(ref State.NumBytesRead, -RollbackSize);
-
-						// Bail out after retrying
-						if(++NumAttempts == MaxRetries)
-						{
-							Interlocked.CompareExchange(ref State.ErrorMessage, String.Format("Failed to download '{0}' to '{1}' after {2} attempts: {3} ({4})", BundleUrl, PackFileName, NumAttempts, CaughtException.Message, CaughtException.GetType().Name), null);
+							Interlocked.CompareExchange(ref State.LastDecompressError, String.Format("Failed to extract '{0}': {1}", OutputFileName, Ex.Message), null);
 							return;
 						}
+						Interlocked.Increment(ref State.NumFilesRead);
 					}
+				}
+
+				// Delete the pack file now that we're finished with it. Doesn't matter much if it fails.
+				try { System.IO.File.Delete(PackFileName); } catch(Exception) { }
+			}
+		}
 
-					// Move the files into their final place
-					foreach(DependencyBlob Blob in PackToBlobs[NextPack.Hash])
+		static void DownloadWorker(string RootPath, ConcurrentQueue<DependencyPack> DownloadQueue, ConcurrentQueue<DependencyPack> DecompressQueue, Dictionary<DependencyPack, string> DownloadFileNames, Dictionary<string, List<DependencyBlob>> PackToBlobs, Dictionary<string, List<DependencyFile>> BlobToFiles, AsyncDownloadState State, int MaxRetries)
+		{
+			int Retries = 0;
+			while(State.NumFilesRead < State.NumFiles)
+			{
+				// Remove the next file from the download queue, or wait before polling again
+				DependencyPack NextPack;
+				if(!DownloadQueue.TryDequeue(out NextPack))
+				{
+					Thread.Sleep(100);
+					continue;
+				}
+
+				// Get the temporary file to download to
+				string PackFileName = DownloadFileNames[NextPack];
+
+				// Format the URL for it
+				string Url = String.Format("http://cdn.unrealengine.com/dependencies/{0}/{1}", NextPack.RemotePath, NextPack.Hash);
+
+				// Try to download the file
+				long RollbackSize = 0;
+				try
+				{
+					// Download the file and queue it for decompression
+					DownloadFileAndVerifyHash(Url, PackFileName, NextPack.Hash, Size => { RollbackSize += Size; Interlocked.Add(ref State.NumBytesRead, Size); });
+					DecompressQueue.Enqueue(NextPack);
+
+					// If we were failing, decrement the number of failing threads
+					if(Retries > MaxRetries)
 					{
-						foreach(DependencyFile File in BlobToFiles[Blob.Hash])
-						{
-							string OutputFileName = Path.Combine(RootPath, File.Name);
-							try
-							{
-								ExtractBlob(PackFileName, Blob, OutputFileName);
-							}
-							catch(Exception Ex)
-							{
-								Interlocked.CompareExchange(ref State.ErrorMessage, String.Format("Failed to extract '{0}': {1}", OutputFileName, Ex.Message), null);
-								return;
-							}
-							Interlocked.Increment(ref State.NumFilesRead);
-						}
+						Interlocked.Decrement(ref State.NumFailingDownloads);
+						Retries = 0;
 					}
 				}
 				catch(Exception Ex)
 				{
-					Interlocked.CompareExchange(ref State.ErrorMessage, Ex.ToString(), null);
-					break;
-				}
-				finally
-				{
-					File.Delete(PackFileName);
+					// Rollback the byte count and add the file back into the download queue
+					Interlocked.Add(ref State.NumBytesRead, -RollbackSize);
+					DownloadQueue.Enqueue(NextPack);
+
+					// If we've retried enough times already, set the error message. 
+					if(Retries++ == MaxRetries)
+					{
+						Interlocked.Increment(ref State.NumFailingDownloads);
+						State.LastDownloadError = String.Format("Failed to download '{0}' to '{1}': {2} ({3})", Url, PackFileName, Ex.Message, Ex.GetType().Name);
+					}
 				}
 			}
 		}
 
 		static void DownloadFileAndVerifyHash(string Url, string PackFileName, string ExpectedHash, NotifyReadDelegate NotifyRead)
 		{
-			// Create the output folder
-			Directory.CreateDirectory(Path.GetDirectoryName(PackFileName));
+			// Create the web request
+			WebRequest Request = WebRequest.Create(Url);
+			Request.Proxy = null;
 
-			// Download the file, decompressing and hashing it as we go
-			SHA1Managed Hasher = new SHA1Managed();
-			using(FileStream OutputStream = File.OpenWrite(PackFileName))
+			// Get the response
+			using(WebResponse Response = Request.GetResponse())
 			{
-				CryptoStream HashOutputStream = new CryptoStream(OutputStream, Hasher, CryptoStreamMode.Write);
-				using (WebClient Client = new WebClient())
+				// Download the file, decompressing and hashing it as we go
+				SHA1Managed Hasher = new SHA1Managed();
+				using(FileStream OutputStream = File.OpenWrite(PackFileName))
 				{
-					using(NotifyReadStream InputStream = new NotifyReadStream(Client.OpenRead(Url), NotifyRead))
+					CryptoStream HashOutputStream = new CryptoStream(OutputStream, Hasher, CryptoStreamMode.Write);
+					using(NotifyReadStream InputStream = new NotifyReadStream(Response.GetResponseStream(), NotifyRead))
 					{
 						GZipStream DecompressedStream = new GZipStream(InputStream, CompressionMode.Decompress, true);
 						DecompressedStream.CopyTo(HashOutputStream);
 					}
+					HashOutputStream.FlushFinalBlock();
 				}
-				HashOutputStream.FlushFinalBlock();
-			}
 
-			// Check the hash was what we expected
-			string Hash = BitConverter.ToString(Hasher.Hash).ToLower().Replace("-", "");
-			if(Hash != ExpectedHash)
-			{
-				throw new InvalidDataException("Incorrect hash value");
+				// Check the hash was what we expected
+				string Hash = BitConverter.ToString(Hasher.Hash).ToLower().Replace("-", "");
+				if(Hash != ExpectedHash)
+				{
+					throw new InvalidDataException(String.Format("Incorrect hash for {0} - expected {1}, got {2}", Url, ExpectedHash, Hash));
+				}
 			}
 		}
 
