Add support for setting executable permissions on Posix platforms from the Git dependency manifest, and clean up some output logging.

diff --git a/Engine/Source/Programs/GitDependencies/DependencyManifest.cs b/Engine/Source/Programs/GitDependencies/DependencyManifest.cs
--- a/Engine/Source/Programs/GitDependencies/DependencyManifest.cs
+++ b/Engine/Source/Programs/GitDependencies/DependencyManifest.cs
@@ -17,6 +17,9 @@ namespace GitDependencies
 
 		[XmlAttribute]
 		public string Hash;
+
+		[XmlAttribute]
+		public bool IsExecutable;
 	}
 
 	[DebuggerDisplay("{Hash}")]
diff --git a/Engine/Source/Programs/GitDependencies/Log.cs b/Engine/Source/Programs/GitDependencies/Log.cs
--- a/Engine/Source/Programs/GitDependencies/Log.cs
+++ b/Engine/Source/Programs/GitDependencies/Log.cs
@@ -38,27 +38,13 @@ namespace GitDependencies
 
 		public static void WriteStatus(string Format, params object[] Args)
 		{
+			// Write the new status, and clear any space after the end of the string if it's shorter
 			string NewStatus = String.Format(Format, Args);
-
-			// Figure out how many characters are in common
-			int NumCommon = 0;
-			while(NumCommon < CurrentStatus.Length && NumCommon < NewStatus.Length && CurrentStatus[NumCommon] == NewStatus[NumCommon])
-			{
-				NumCommon++;
-			}
-
-			// Go back to the point that they diverge and write the new string
-			Console.Write(new String('\b', CurrentStatus.Length - NumCommon));
-			Console.Write(NewStatus.Substring(NumCommon));
-
-			// If the new message is shorter, clear the rest of the line
+			Console.Write("\r" + NewStatus);
 			if(NewStatus.Length < CurrentStatus.Length)
 			{
-				Console.Write(new String(' ', CurrentStatus.Length - NewStatus.Length));
-				Console.Write(new String('\b', CurrentStatus.Length - NewStatus.Length));
+				Console.Write(new string(' ', CurrentStatus.Length - NewStatus.Length) + "\r" + NewStatus);
 			}
-
-			// Update the status message
 			CurrentStatus = NewStatus;
 		}
 
diff --git a/Engine/Source/Programs/GitDependencies/NotifyReadStream.cs b/Engine/Source/Programs/GitDependencies/NotifyReadStream.cs
--- a/Engine/Source/Programs/GitDependencies/NotifyReadStream.cs
+++ b/Engine/Source/Programs/GitDependencies/NotifyReadStream.cs
@@ -3,7 +3,6 @@ using System.Collections.Generic;
 using System.IO;
 using System.Linq;
 using System.Text;
-using System.Threading.Tasks;
 
 namespace GitDependencies
 {
diff --git a/Engine/Source/Programs/GitDependencies/Program.cs b/Engine/Source/Programs/GitDependencies/Program.cs
--- a/Engine/Source/Programs/GitDependencies/Program.cs
+++ b/Engine/Source/Programs/GitDependencies/Program.cs
@@ -170,11 +170,11 @@ namespace GitDependencies
 			// Start scanning on the working directory 
 			if(ExcludeFolders.Count > 0)
 			{
-				Log.WriteLine("Checking binary dependencies (excluding {0})...", String.Join(", ", ExcludeFolders));
+				Log.WriteLine("Updating UE4 binary dependencies (excluding {0})...", String.Join(", ", ExcludeFolders));
 			}
 			else
 			{
-				Log.WriteLine("Checking binary dependencies...");
+				Log.WriteLine("Updating UE4 binary dependencies...");
 			}
 
 			// Figure out the path to the working manifest
@@ -365,12 +365,6 @@ namespace GitDependencies
 				return false;
 			}
 
-			// Log out if we have any up-to-date files
-			if(NewWorkingManifest.Files.Count > FilesToDownload.Count)
-			{
-				Log.WriteLine("{0} files up to date.", NewWorkingManifest.Files.Count - FilesToDownload.Count);
-			}
-
 			// If there's nothing to do, just print a simpler message and exit early
 			if(FilesToDownload.Count > 0)
 			{
@@ -397,6 +391,96 @@ namespace GitDependencies
 					return false;
 				}
 			}
+
+			// Update all the executable permissions
+			if(!SetExecutablePermissions(RootPath, TargetFiles.Values))
+			{
+				return false;
+			}
+
+			return true;
+		}
+
+		static bool SetExecutablePermissions(string RootDir, IEnumerable<DependencyFile> Files)
+		{
+			// Try to load the Mono Posix assembly. If it doesn't exist, we're on Windows.
+			Assembly MonoPosix;
+			try
+			{
+				MonoPosix = Assembly.Load("Mono.Posix, Version=4.0.0.0, Culture=neutral, PublicKeyToken=0738eb9f132ed756");
+			}
+			catch(FileNotFoundException)
+			{
+				return true;
+			}
+
+			// Dynamically find all the types and methods for Syscall.stat and Syscall.chmod
+			Type SyscallType = MonoPosix.GetType("Mono.Unix.Native.Syscall");
+			if(SyscallType == null)
+			{
+				Log.WriteError("Couldn't find Syscall type");
+				return false;
+			}
+            MethodInfo StatMethod = SyscallType.GetMethod ("stat");
+			if(StatMethod == null)
+			{
+				Log.WriteError("Couldn't find Mono.Unix.Native.Syscall.stat method");
+				return false;
+			}
+			MethodInfo ChmodMethod = SyscallType.GetMethod("chmod");
+			if(ChmodMethod == null)
+			{
+				Log.WriteError("Couldn't find Mono.Unix.Native.Syscall.chmod method");
+				return false;
+			}
+			Type StatType = MonoPosix.GetType("Mono.Unix.Native.Stat");
+			if(StatType == null)
+			{
+				Log.WriteError("Couldn't find Mono.Unix.Native.Stat type");
+				return false;
+			}
+			FieldInfo StatModeField = StatType.GetField("st_mode");
+			if(StatModeField == null)
+			{
+				Log.WriteError("Couldn't find Mono.Unix.Native.Stat.st_mode field");
+				return false;
+			}
+
+			// Update all the executable permissions
+			const uint ExecutableBits = (1 << 0) | (1 << 3) | (1 << 6);
+			foreach(DependencyFile File in Files)
+			{
+				if(File.IsExecutable)
+				{
+					string FileName = Path.Combine(RootDir, File.Name);
+
+					// Call Syscall.stat(Filename, out Stat)
+					object[] StatArgs = new object[]{ FileName, null };
+					int StatResult = (int)StatMethod.Invoke(null, StatArgs);
+					if(StatResult != 0)
+					{
+						Log.WriteError("Stat() call for {0} failed with error {1}", File.Name, StatResult);
+						return false;
+					}
+
+					// Get the current permissions
+					uint CurrentPermissions = (uint)StatModeField.GetValue(StatArgs[1]);
+
+					// The desired permissions should be executable for every read group
+					uint NewPermissions = CurrentPermissions | ((CurrentPermissions >> 2) & ExecutableBits);
+
+					// Update them if they don't match
+					if (CurrentPermissions != NewPermissions)
+					{
+						int ChmodResult = (int)ChmodMethod.Invoke(null, new object[]{ FileName, NewPermissions });
+						if(ChmodResult != 0)
+						{
+							Log.WriteError("Chmod() call for {0} failed with error {1}", File.Name, ChmodResult);
+							return false;
+						}
+					}
+				}
+			}
 			return true;
 		}
 
@@ -477,13 +561,14 @@ namespace GitDependencies
 
 			// Tick the status message until we've finished or ended with an error. Use a circular buffer to average out the speed over time.
 			long[] NumBytesReadBuffer = new long[60];
-			for(int BufferIdx = 0; State.NumFilesRead < State.NumFiles && State.NumFailingDownloads < NumThreads && State.LastDecompressError == null; BufferIdx = (BufferIdx + 1) % NumBytesReadBuffer.Length)
+			for(int BufferIdx = 0, NumFilesReportedRead = 0; NumFilesReportedRead < State.NumFiles && State.NumFailingDownloads < NumThreads && State.LastDecompressError == null; BufferIdx = (BufferIdx + 1) % NumBytesReadBuffer.Length)
 			{
 				const int TickInterval = 100;
 
 				long NumBytesRead = Interlocked.Read(ref State.NumBytesRead);
 				float NumBytesPerSecond = (float)Math.Max(NumBytesRead - NumBytesReadBuffer[BufferIdx], 0) * 1000.0f / (NumBytesReadBuffer.Length * TickInterval);
-				Log.WriteStatus("Staged {0}/{1} files ({2:0.0}/{3:0.0}mb; {4:0.00}mb/s; {5}%)...", State.NumFilesRead, State.NumFiles, (NumBytesRead / (1024.0 * 1024.0)) + 0.0999999, (NumBytesTotal / (1024.0 * 1024.0)) + 0.0999999, (NumBytesPerSecond / (1024.0 * 1024.0)) + 0.00999999, (NumBytesRead * 100) / NumBytesTotal);
+				NumFilesReportedRead = State.NumFilesRead;
+				Log.WriteStatus("Received {0}/{1} files ({2:0.0}/{3:0.0}mb; {4:0.00}mb/s; {5}%)...", NumFilesReportedRead, State.NumFiles, (NumBytesRead / (1024.0 * 1024.0)) + 0.0999999, (NumBytesTotal / (1024.0 * 1024.0)) + 0.0999999, (NumBytesPerSecond / (1024.0 * 1024.0)) + 0.0099, (NumBytesRead * 100) / NumBytesTotal);
 				NumBytesReadBuffer[BufferIdx] = NumBytesRead;
 
 				Thread.Sleep(TickInterval);
