Clean up help messages for GitDependencies to fit within 80 character wide terminals.

diff --git a/Engine/Source/Programs/GitDependencies/Program.cs b/Engine/Source/Programs/GitDependencies/Program.cs
--- a/Engine/Source/Programs/GitDependencies/Program.cs
+++ b/Engine/Source/Programs/GitDependencies/Program.cs
@@ -128,20 +128,18 @@ namespace GitDependencies
 				Log.WriteLine("   --all                         Sync all folders");
 				Log.WriteLine("   --include=<X>                 Include binaries in folders called <X>");
 				Log.WriteLine("   --exclude=<X>                 Exclude binaries in folders called <X>");
-				Log.WriteLine("   --prompt                      Prompts for whether to overwrite modified workspace files");
-				Log.WriteLine("   --force                       Overwrite modified dependency files in the workspace");
-				Log.WriteLine("   --root=<PATH>                 Specifies the path to the directory to sync with");
-				Log.WriteLine("   --threads=X                   Use X threads when downloading new files");
-				Log.WriteLine("   --dry-run                     Print a list of outdated files, but don't do anything");
-				Log.WriteLine("   --max-retries                 Set the maximum number of retries for downloading files");
+				Log.WriteLine("   --prompt                      Prompt before overwriting modified files");
+				Log.WriteLine("   --force                       Always overwrite modified files");
+				Log.WriteLine("   --root=<PATH>                 Set the repository directory to be sync");
+				Log.WriteLine("   --threads=<N>                 Use N threads when downloading new files");
+				Log.WriteLine("   --dry-run                     Print a list of outdated files and exit");
+				Log.WriteLine("   --max-retries                 Set the maximum number of retries for each file");
 				Log.WriteLine("   --proxy=<URI>                 Set http proxy URL or address");
 				Log.WriteLine("   --proxy-user=<username>       Set proxy username");
 				Log.WriteLine("   --proxy-password=<password>   Set proxy password");
-				if (ExcludeFolders.Count > 0)
-				{
-					Log.WriteLine();
-					Log.WriteLine("Current excluded folders: {0}", String.Join(", ", ExcludeFolders));
-				}
+				Log.WriteLine();
+				Log.WriteLine("Detected settings:");
+				Log.WriteLine("   Excluded folders: {0}", (ExcludeFolders.Count == 0)? "none" : String.Join(", ", ExcludeFolders));
 				return 0;
 			}
 
