Condense the proxy server settings into a single parameter.

diff --git a/Engine/Source/Programs/GitDependencies/Program.cs b/Engine/Source/Programs/GitDependencies/Program.cs
--- a/Engine/Source/Programs/GitDependencies/Program.cs
+++ b/Engine/Source/Programs/GitDependencies/Program.cs
@@ -68,13 +68,27 @@ namespace GitDependencies
 
 			// Setup network proxy from argument list or environment variable
 			string ProxyUrl = ParseParameter(ArgsList, "-proxy=", null);
-			string ProxyUsername = ParseParameter(ArgsList, "-proxy-user=", null);
-			string ProxyPassword = ParseParameter(ArgsList, "-proxy-password=", null);
 			if(String.IsNullOrEmpty(ProxyUrl))
 			{
 				ProxyUrl = Environment.GetEnvironmentVariable("HTTP_PROXY");
 			}
 
+			// Create a URI for the proxy. If there's no included username/password, accept them as separate parameters for legacy reasons.
+			Uri Proxy = null;
+			if(!String.IsNullOrEmpty(ProxyUrl))
+			{
+				UriBuilder ProxyBuilder = new UriBuilder(ProxyUrl);
+				if(String.IsNullOrEmpty(ProxyBuilder.UserName))
+				{
+					ProxyBuilder.UserName = ParseParameter(ArgsList, "-proxy-user=", null);
+				}
+				if(String.IsNullOrEmpty(ProxyBuilder.Password))
+				{
+					ProxyBuilder.Password = ParseParameter(ArgsList, "-proxy-password=", null);
+				}
+				Proxy = ProxyBuilder.Uri;
+			}
+
 			// Parse all the default exclude filters
 			HashSet<string> ExcludeFolders = new HashSet<string>(StringComparer.CurrentCultureIgnoreCase);
 			if(!ParseSwitch(ArgsList, "-all"))
@@ -134,12 +148,11 @@ namespace GitDependencies
 				Log.WriteLine("   --threads=<N>                 Use N threads when downloading new files");
 				Log.WriteLine("   --dry-run                     Print a list of outdated files and exit");
 				Log.WriteLine("   --max-retries                 Set the maximum number of retries for each file");
-				Log.WriteLine("   --proxy=<URI>                 Set http proxy URL or address");
-				Log.WriteLine("   --proxy-user=<username>       Set proxy username");
-				Log.WriteLine("   --proxy-password=<password>   Set proxy password");
+				Log.WriteLine("   --proxy=<user:password@url>   Sets the HTTP proxy address and credentials");
 				Log.WriteLine();
 				Log.WriteLine("Detected settings:");
 				Log.WriteLine("   Excluded folders: {0}", (ExcludeFolders.Count == 0)? "none" : String.Join(", ", ExcludeFolders));
+				Log.WriteLine("   Proxy server: {0}", (Proxy == null)? "none" : Proxy.ToString());
 				return 0;
 			}
 
@@ -147,7 +160,7 @@ namespace GitDependencies
 			Console.CancelKeyPress += delegate { Log.FlushStatus(); };
 
 			// Update the tree. Make sure we clear out the status line if we quit for any reason (eg. ctrl-c)
-			if(!UpdateWorkingTree(bDryRun, RootPath, ExcludeFolders, NumThreads, MaxRetries, ProxyUrl, ProxyUsername, ProxyPassword, Overwrite))
+			if(!UpdateWorkingTree(bDryRun, RootPath, ExcludeFolders, NumThreads, MaxRetries, Proxy, Overwrite))
 			{
 				return 1;
 			}
@@ -195,7 +208,7 @@ namespace GitDependencies
 			}
 		}
 
-		static bool UpdateWorkingTree(bool bDryRun, string RootPath, HashSet<string> ExcludeFolders, int NumThreads, int MaxRetries, string ProxyUrl, string ProxyUsername, string ProxyPassword, OverwriteMode Overwrite)
+		static bool UpdateWorkingTree(bool bDryRun, string RootPath, HashSet<string> ExcludeFolders, int NumThreads, int MaxRetries, Uri Proxy, OverwriteMode Overwrite)
 		{
 			// Start scanning on the working directory 
 			if(ExcludeFolders.Count > 0)
@@ -443,7 +456,7 @@ namespace GitDependencies
 			if(FilesToDownload.Count > 0)
 			{
 				// Download all the new dependencies
-				if(!DownloadDependencies(RootPath, FilesToDownload, TargetBlobs.Values, TargetPacks.Values, NumThreads, MaxRetries, ProxyUrl, ProxyUsername, ProxyPassword))
+				if(!DownloadDependencies(RootPath, FilesToDownload, TargetBlobs.Values, TargetPacks.Values, NumThreads, MaxRetries, Proxy))
 				{
 					return false;
 				}
@@ -570,7 +583,7 @@ namespace GitDependencies
 			return false;
 		}
 
-		static bool DownloadDependencies(string RootPath, IEnumerable<DependencyFile> RequiredFiles, IEnumerable<DependencyBlob> Blobs, IEnumerable<DependencyPack> Packs, int NumThreads, int MaxRetries, string ProxyUrl, string ProxyUsername, string ProxyPassword)
+		static bool DownloadDependencies(string RootPath, IEnumerable<DependencyFile> RequiredFiles, IEnumerable<DependencyBlob> Blobs, IEnumerable<DependencyPack> Packs, int NumThreads, int MaxRetries, Uri Proxy)
 		{
 			// Make sure we can actually open the right number of connections
 			ServicePointManager.DefaultConnectionLimit = NumThreads;
@@ -625,7 +638,7 @@ namespace GitDependencies
 			Thread[] WorkerThreads = new Thread[NumThreads];
 			for(int Idx = 0; Idx < NumThreads; Idx++)
 			{
-				WorkerThreads[Idx] = new Thread(x => DownloadWorker(RootPath, DownloadQueue, DecompressQueue, DownloadFileNames, PackToBlobs, BlobToFiles, State, MaxRetries, ProxyUrl, ProxyUsername, ProxyPassword));
+				WorkerThreads[Idx] = new Thread(x => DownloadWorker(RootPath, DownloadQueue, DecompressQueue, DownloadFileNames, PackToBlobs, BlobToFiles, State, MaxRetries, Proxy));
 				WorkerThreads[Idx].Start();
 			}
 
@@ -713,7 +726,7 @@ namespace GitDependencies
 			}
 		}
 
-		static void DownloadWorker(string RootPath, ConcurrentQueue<DependencyPack> DownloadQueue, ConcurrentQueue<DependencyPack> DecompressQueue, Dictionary<DependencyPack, string> DownloadFileNames, Dictionary<string, List<DependencyBlob>> PackToBlobs, Dictionary<string, List<DependencyFile>> BlobToFiles, AsyncDownloadState State, int MaxRetries, string ProxyUrl, string ProxyUsername, string ProxyPassword)
+		static void DownloadWorker(string RootPath, ConcurrentQueue<DependencyPack> DownloadQueue, ConcurrentQueue<DependencyPack> DecompressQueue, Dictionary<DependencyPack, string> DownloadFileNames, Dictionary<string, List<DependencyBlob>> PackToBlobs, Dictionary<string, List<DependencyFile>> BlobToFiles, AsyncDownloadState State, int MaxRetries, Uri Proxy)
 		{
 			int Retries = 0;
 			while(State.NumFilesRead < State.NumFiles)
@@ -737,7 +750,7 @@ namespace GitDependencies
 				try
 				{
 					// Download the file and queue it for decompression
-					DownloadFileAndVerifyHash(Url, ProxyUrl, ProxyUsername, ProxyPassword, PackFileName, NextPack.Hash, Size => { RollbackSize += Size; Interlocked.Add(ref State.NumBytesRead, Size); });
+					DownloadFileAndVerifyHash(Url, Proxy, PackFileName, NextPack.Hash, Size => { RollbackSize += Size; Interlocked.Add(ref State.NumBytesRead, Size); });
 					DecompressQueue.Enqueue(NextPack);
 
 					// If we were failing, decrement the number of failing threads
@@ -763,22 +776,18 @@ namespace GitDependencies
 			}
 		}
 
-		static void DownloadFileAndVerifyHash(string Url, string ProxyUrl, string ProxyUsername, string ProxyPassword, string PackFileName, string ExpectedHash, NotifyReadDelegate NotifyRead)
+		static void DownloadFileAndVerifyHash(string Url, Uri Proxy, string PackFileName, string ExpectedHash, NotifyReadDelegate NotifyRead)
 		{
 			// Create the web request
 			WebRequest Request = WebRequest.Create(Url);
-			if(String.IsNullOrEmpty(ProxyUrl))
+			if(Proxy == null)
 			{
 				Request.Proxy = null;
 			}
-			else if (String.IsNullOrEmpty(ProxyUsername))
+			else
 			{
-                Request.Proxy = new WebProxy(ProxyUrl);
+				Request.Proxy = new WebProxy(Proxy, true, null, MakeCredentialsFromUri(Proxy));
 			}
-            else
-            {
-                Request.Proxy = new WebProxy(ProxyUrl, true, null, new NetworkCredential(ProxyUsername, ProxyPassword));
-            }
 
 			// Get the response
 			using(WebResponse Response = Request.GetResponse())
@@ -805,6 +814,19 @@ namespace GitDependencies
 			}
 		}
 
+		static NetworkCredential MakeCredentialsFromUri(Uri Address)
+		{
+			if(!String.IsNullOrEmpty(Address.UserInfo))
+			{
+				int Index = Address.UserInfo.IndexOf(':');
+				if(Index != -1)
+				{
+					return new NetworkCredential(Address.UserInfo.Substring(0, Index), Address.UserInfo.Substring(Index + 1));
+				}
+			}
+			return null;
+		}
+
 		static void ExtractBlob(string PackFileName, DependencyBlob Blob, string OutputFileName)
 		{
 			// Create the output folder
