Add a utility function for a line of unformatted status text.

diff --git a/Engine/Source/Programs/GitDependencies/Log.cs b/Engine/Source/Programs/GitDependencies/Log.cs
--- a/Engine/Source/Programs/GitDependencies/Log.cs
+++ b/Engine/Source/Programs/GitDependencies/Log.cs
@@ -38,10 +38,9 @@ namespace GitDependencies
 			Console.ResetColor();
 		}
 
-		public static void WriteStatus(string Format, params object[] Args)
+		public static void WriteStatus(string NewStatus)
 		{
 			// Write the new status, and clear any space after the end of the string if it's shorter
-			string NewStatus = String.Format(Format, Args);
 			Console.Write("\r" + NewStatus);
 			if(NewStatus.Length < CurrentStatus.Length)
 			{
@@ -50,6 +49,11 @@ namespace GitDependencies
 			CurrentStatus = NewStatus;
 		}
 
+		public static void WriteStatus(string Format, params object[] Args)
+		{
+			WriteStatus(String.Format(Format, Args));
+		}
+
 		public static void FlushStatus()
 		{
 			if(CurrentStatus.Length > 0)
