PR #749: Added cache support for GitDependencies. (Contributed by [REDACTED_PERSON])

https://github.com/EpicGames/UnrealEngine/pull/749

diff --git a/Engine/Source/Programs/GitDependencies/DependencyManifest.cs b/Engine/Source/Programs/GitDependencies/DependencyManifest.cs
--- a/Engine/Source/Programs/GitDependencies/DependencyManifest.cs
+++ b/Engine/Source/Programs/GitDependencies/DependencyManifest.cs
@@ -58,6 +58,12 @@ namespace GitDependencies
 
 	public class DependencyManifest
 	{
+		[XmlAttribute]
+		public string BaseUrl;
+
+		[XmlAttribute]
+		public bool IgnoreProxy;
+
 		[XmlArrayItem("File")]
 		public DependencyFile[] Files = new DependencyFile[0];
 
diff --git a/Engine/Source/Programs/GitDependencies/ForkReadStream.cs b/Engine/Source/Programs/GitDependencies/ForkReadStream.cs
new file mode 100644
--- /dev/null
+++ b/Engine/Source/Programs/GitDependencies/ForkReadStream.cs
@@ -0,0 +1,98 @@
+﻿// Copyright 1998-2015 Epic Games, Inc. All Rights Reserved.
+
+using System;
+using System.Collections.Generic;
+using System.IO;
+using System.Linq;
+using System.Text;
+
+namespace GitDependencies
+{
+	class ForkReadStream : Stream
+	{
+		Stream Inner;
+		Stream ForkWrite;
+
+		public ForkReadStream(Stream InInner, Stream InForkWrite)
+		{
+			Inner = InInner;
+			ForkWrite = InForkWrite;
+		}
+
+		protected override void Dispose(bool Disposing)
+		{
+			if (ForkWrite != null)
+			{
+				ForkWrite.Dispose();
+				ForkWrite = null;
+			}
+			if (Inner != null)
+			{
+				Inner.Dispose();
+				Inner = null;
+			}
+		}
+
+		public override bool CanRead
+		{
+			get { return true; }
+		}
+
+		public override bool CanWrite
+		{
+			get { return false; }
+		}
+
+		public override bool CanSeek
+		{
+			get { return false; }
+		}
+
+		public override long Position
+		{
+			get
+			{
+				return Inner.Position;
+			}
+			set
+			{
+				throw new NotImplementedException();
+			}
+		}
+
+		public override long Length
+		{
+			get { return Inner.Length; }
+		}
+
+		public override void SetLength(long Value)
+		{
+			throw new NotImplementedException();
+		}
+
+		public override int Read(byte[] Buffer, int Offset, int Count)
+		{
+			int SizeRead = Inner.Read(Buffer, Offset, Count);
+			if (SizeRead > 0)
+			{
+				ForkWrite.Write(Buffer, Offset, SizeRead);
+			}
+			return SizeRead;
+		}
+
+		public override void Write(byte[] Buffer, int Offset, int Count)
+		{
+			throw new NotImplementedException();
+		}
+
+		public override long Seek(long Offset, SeekOrigin Origin)
+		{
+			throw new NotImplementedException();
+		}
+
+		public override void Flush()
+		{
+			ForkWrite.Flush();
+		}
+	}
+}
diff --git a/Engine/Source/Programs/GitDependencies/GitDependencies.csproj b/Engine/Source/Programs/GitDependencies/GitDependencies.csproj
--- a/Engine/Source/Programs/GitDependencies/GitDependencies.csproj
+++ b/Engine/Source/Programs/GitDependencies/GitDependencies.csproj
@@ -45,6 +45,7 @@
       <Link>Properties\MetaData.cs</Link>
     </Compile>
     <Compile Include="DependencyManifest.cs" />
+    <Compile Include="ForkReadStream.cs" />
     <Compile Include="Log.cs" />
     <Compile Include="NotifyReadStream.cs" />
     <Compile Include="Program.cs" />
diff --git a/Engine/Source/Programs/GitDependencies/NotifyReadStream.cs b/Engine/Source/Programs/GitDependencies/NotifyReadStream.cs
--- a/Engine/Source/Programs/GitDependencies/NotifyReadStream.cs
+++ b/Engine/Source/Programs/GitDependencies/NotifyReadStream.cs
@@ -14,11 +14,13 @@ namespace GitDependencies
 	{
 		Stream Inner;
 		NotifyReadDelegate NotifyRead;
+		long TotalRead;
 
 		public NotifyReadStream(Stream InInner, NotifyReadDelegate InNotifyRead)
 		{
 			Inner = InInner;
 			NotifyRead = InNotifyRead;
+			TotalRead = 0;
 		}
 
 		protected override void Dispose(bool Disposing)
@@ -49,7 +51,7 @@ namespace GitDependencies
 		{
 			get
 			{
-				return Inner.Position;
+				return TotalRead;
 			}
 			set
 			{
@@ -71,6 +73,7 @@ namespace GitDependencies
 		{
 			int SizeRead = Inner.Read(Buffer, Offset, Count);
 			NotifyRead(SizeRead);
+			TotalRead += SizeRead;
 			return SizeRead;
 		}
 
diff --git a/Engine/Source/Programs/GitDependencies/Program.cs b/Engine/Source/Programs/GitDependencies/Program.cs
--- a/Engine/Source/Programs/GitDependencies/Program.cs
+++ b/Engine/Source/Programs/GitDependencies/Program.cs
@@ -24,8 +24,9 @@ namespace GitDependencies
 			public int NumFiles;
 			public int NumFilesRead;
 			public long NumBytesRead;
-			public string LastDecompressError;
-			public int NumFailingDownloads;
+			public long NumBytesTotal;
+			public long NumBytesCached;
+			public int NumFailingOrIdleDownloads;
 			public string LastDownloadError;
 		}
 
@@ -36,6 +37,25 @@ namespace GitDependencies
 			Force,
 		}
 
+		struct DependencyPackInfo
+		{
+			public DependencyManifest Manifest;
+			public DependencyPack Pack;
+
+			public DependencyPackInfo(DependencyManifest Manifest, DependencyPack Pack) {
+				this.Manifest = Manifest;
+				this.Pack = Pack;
+			}
+		}
+
+		class InvalidCacheException : Exception
+		{
+			public InvalidCacheException(string Message, Exception InnerException)
+				: base(Message, InnerException)
+			{
+			}
+		}
+
 		const string IncomingFileSuffix = ".incoming";
 		const string TempManifestExtension = ".tmp";
 		
@@ -55,6 +75,39 @@ namespace GitDependencies
 			bool bHelp = ParseSwitch(ArgsList, "-help");
 			string RootPath = ParseParameter(ArgsList, "-root=", Path.GetFullPath(Path.Combine(Path.GetDirectoryName(Assembly.GetExecutingAssembly().Location), "../../..")));
 
+			//
+			// Cache path
+			// 
+			// * If specified flag --no-cache - do not use the cache;
+			// * If specified flag --cache=<PATH> - use the path <PATH>;
+			// * If specified enviroment variable UE4_GITDEPS, then use path from it;
+			// * If there is a .git directory above from --root=<PATH>, then use the path of the form .git/ue4-gitdeps/;
+			// * In other case - do not use the cache.
+			//
+			string CachePath = null;
+			if (!ParseSwitch(ArgsList, "-no-cache"))
+			{
+				string CachePathParam = ParseParameter(ArgsList, "-cache=", System.Environment.GetEnvironmentVariable("UE4_GITDEPS"));
+				if (String.IsNullOrEmpty(CachePathParam))
+				{
+					string CheckPath = Path.GetFullPath(RootPath);
+					while (CheckPath != null)
+					{
+						string GitPath = Path.Combine(CheckPath, ".git");
+						if (Directory.Exists(GitPath))
+						{
+							CachePath = Path.Combine(GitPath, "ue4-gitdeps");
+							break;
+						}
+						CheckPath = Path.GetDirectoryName(CheckPath);
+					}
+				}
+				else
+				{
+					CachePath = Path.GetFullPath(CachePathParam);
+				}
+			}
+
 			// Parse the overwrite mode
 			OverwriteMode Overwrite = OverwriteMode.Unchanged;
 			if(ParseSwitch(ArgsList, "-prompt"))
@@ -149,10 +202,13 @@ namespace GitDependencies
 				Log.WriteLine("   --dry-run                     Print a list of outdated files and exit");
 				Log.WriteLine("   --max-retries                 Set the maximum number of retries for each file");
 				Log.WriteLine("   --proxy=<user:password@url>   Sets the HTTP proxy address and credentials");
+				Log.WriteLine("   --cache=<PATH>                Specifies a custom path for the download cache");
+				Log.WriteLine("   --no-cache                    Disable caching of downloaded files");
 				Log.WriteLine();
 				Log.WriteLine("Detected settings:");
 				Log.WriteLine("   Excluded folders: {0}", (ExcludeFolders.Count == 0)? "none" : String.Join(", ", ExcludeFolders));
 				Log.WriteLine("   Proxy server: {0}", (Proxy == null)? "none" : Proxy.ToString());
+				Log.WriteLine("   Current cache path: {0}", CachePath != null ? CachePath : "cache disabled");
 				return 0;
 			}
 
@@ -160,7 +216,7 @@ namespace GitDependencies
 			Console.CancelKeyPress += delegate { Log.FlushStatus(); };
 
 			// Update the tree. Make sure we clear out the status line if we quit for any reason (eg. ctrl-c)
-			if(!UpdateWorkingTree(bDryRun, RootPath, ExcludeFolders, NumThreads, MaxRetries, Proxy, Overwrite))
+			if(!UpdateWorkingTree(bDryRun, RootPath, ExcludeFolders, NumThreads, MaxRetries, Proxy, Overwrite, CachePath))
 			{
 				return 1;
 			}
@@ -208,7 +264,7 @@ namespace GitDependencies
 			}
 		}
 
-		static bool UpdateWorkingTree(bool bDryRun, string RootPath, HashSet<string> ExcludeFolders, int NumThreads, int MaxRetries, Uri Proxy, OverwriteMode Overwrite)
+		static bool UpdateWorkingTree(bool bDryRun, string RootPath, HashSet<string> ExcludeFolders, int NumThreads, int MaxRetries, Uri Proxy, OverwriteMode Overwrite, string CachePath)
 		{
 			// Start scanning on the working directory 
 			if(ExcludeFolders.Count > 0)
@@ -253,7 +309,7 @@ namespace GitDependencies
 			// Find all the manifests and push them into dictionaries
 			Dictionary<string, DependencyFile> TargetFiles = new Dictionary<string,DependencyFile>(StringComparer.InvariantCultureIgnoreCase);
 			Dictionary<string, DependencyBlob> TargetBlobs = new Dictionary<string,DependencyBlob>(StringComparer.InvariantCultureIgnoreCase);
-			Dictionary<string, DependencyPack> TargetPacks = new Dictionary<string,DependencyPack>(StringComparer.InvariantCultureIgnoreCase);
+			Dictionary<string, DependencyPackInfo> TargetPacks = new Dictionary<string, DependencyPackInfo>(StringComparer.InvariantCultureIgnoreCase);
 			foreach(string BaseFolder in Directory.EnumerateDirectories(RootPath))
 			{
 				string BuildFolder = Path.Combine(BaseFolder, "Build");
@@ -264,26 +320,26 @@ namespace GitDependencies
 						// Ignore any dotfiles; Mac creates them on non-unix partitions to store permission info.
 						if(!Path.GetFileName(ManifestFileName).StartsWith("."))
 						{
-							// Read this manifest
-							DependencyManifest NewTargetManifest;
-							if(!ReadXmlObject(ManifestFileName, out NewTargetManifest))
-							{
-								return false;
-							}
-
-							// Add all the files, blobs and packs into the shared dictionaries
-							foreach(DependencyFile NewFile in NewTargetManifest.Files)
-							{
-								TargetFiles[NewFile.Name] = NewFile;
-							}
-							foreach(DependencyBlob NewBlob in NewTargetManifest.Blobs)
-							{
-								TargetBlobs[NewBlob.Hash] = NewBlob;
-							}
-							foreach(DependencyPack NewPack in NewTargetManifest.Packs)
-							{
-								TargetPacks[NewPack.Hash] = NewPack;
-							}
+						    // Read this manifest
+						    DependencyManifest NewTargetManifest;
+						    if(!ReadXmlObject(ManifestFileName, out NewTargetManifest))
+						    {
+							    return false;
+						    }
+    
+						    // Add all the files, blobs and packs into the shared dictionaries
+						    foreach(DependencyFile NewFile in NewTargetManifest.Files)
+						    {
+							    TargetFiles[NewFile.Name] = NewFile;
+						    }
+						    foreach(DependencyBlob NewBlob in NewTargetManifest.Blobs)
+						    {
+							    TargetBlobs[NewBlob.Hash] = NewBlob;
+						    }
+						    foreach(DependencyPack NewPack in NewTargetManifest.Packs)
+						    {
+							    TargetPacks[NewPack.Hash] = new DependencyPackInfo(NewTargetManifest, NewPack);
+						    }
 						}
 					}
 				}
@@ -456,7 +512,7 @@ namespace GitDependencies
 			if(FilesToDownload.Count > 0)
 			{
 				// Download all the new dependencies
-				if(!DownloadDependencies(RootPath, FilesToDownload, TargetBlobs.Values, TargetPacks.Values, NumThreads, MaxRetries, Proxy))
+				if(!DownloadDependencies(RootPath, FilesToDownload, TargetBlobs.Values, TargetPacks.Values, NumThreads, MaxRetries, Proxy, CachePath))
 				{
 					return false;
 				}
@@ -508,7 +564,7 @@ namespace GitDependencies
 				Log.WriteError("Couldn't find Syscall type");
 				return false;
 			}
-            MethodInfo StatMethod = SyscallType.GetMethod ("stat");
+			MethodInfo StatMethod = SyscallType.GetMethod ("stat");
 			if(StatMethod == null)
 			{
 				Log.WriteError("Couldn't find Mono.Unix.Native.Syscall.stat method");
@@ -583,7 +639,7 @@ namespace GitDependencies
 			return false;
 		}
 
-		static bool DownloadDependencies(string RootPath, IEnumerable<DependencyFile> RequiredFiles, IEnumerable<DependencyBlob> Blobs, IEnumerable<DependencyPack> Packs, int NumThreads, int MaxRetries, Uri Proxy)
+		static bool DownloadDependencies(string RootPath, IEnumerable<DependencyFile> RequiredFiles, IEnumerable<DependencyBlob> Blobs, IEnumerable<DependencyPackInfo> Packs, int NumThreads, int MaxRetries, Uri Proxy, string CachePath)
 		{
 			// Make sure we can actually open the right number of connections
 			ServicePointManager.DefaultConnectionLimit = NumThreads;
@@ -618,67 +674,51 @@ namespace GitDependencies
 			}
 
 			// Find all the required packs
-			DependencyPack[] RequiredPacks = Packs.Where(x => PackToBlobs.ContainsKey(x.Hash)).ToArray();
-
-			// Get temporary filenames for all the files we're going to download
-			Dictionary<DependencyPack, string> DownloadFileNames = new Dictionary<DependencyPack,string>();
-			foreach(DependencyPack Pack in RequiredPacks)
-			{
-				DownloadFileNames.Add(Pack, Path.GetTempFileName());
-			}
+			DependencyPackInfo[] RequiredPacks = Packs.Where(x => PackToBlobs.ContainsKey(x.Pack.Hash)).ToArray();
 
 			// Setup the async state
 			AsyncDownloadState State = new AsyncDownloadState();
 			State.NumFiles = RequiredFiles.Count();
-			long NumBytesTotal = RequiredPacks.Sum(x => x.CompressedSize);
-			ConcurrentQueue<DependencyPack> DownloadQueue = new ConcurrentQueue<DependencyPack>(RequiredPacks);
-			ConcurrentQueue<DependencyPack> DecompressQueue = new ConcurrentQueue<DependencyPack>();
+			State.NumBytesTotal = RequiredPacks.Sum(x => x.Pack.CompressedSize);
+			ConcurrentQueue<DependencyPackInfo> DownloadQueue = new ConcurrentQueue<DependencyPackInfo>(RequiredPacks);
 
 			// Create all the worker threads
 			Thread[] WorkerThreads = new Thread[NumThreads];
 			for(int Idx = 0; Idx < NumThreads; Idx++)
 			{
-				WorkerThreads[Idx] = new Thread(x => DownloadWorker(RootPath, DownloadQueue, DecompressQueue, DownloadFileNames, PackToBlobs, BlobToFiles, State, MaxRetries, Proxy));
+				WorkerThreads[Idx] = new Thread(x => DownloadWorker(RootPath, CachePath, DownloadQueue, PackToBlobs, BlobToFiles, State, MaxRetries, Proxy));
 				WorkerThreads[Idx].Start();
 			}
 
-			// Create the decompression thread
-			Thread DecompressionThread = new Thread(x => DecompressWorker(RootPath, DecompressQueue, DownloadFileNames, PackToBlobs, BlobToFiles, State));
-			DecompressionThread.Start();
-
 			// Tick the status message until we've finished or ended with an error. Use a circular buffer to average out the speed over time.
 			long[] NumBytesReadBuffer = new long[60];
-			for(int BufferIdx = 0, NumFilesReportedRead = 0; NumFilesReportedRead < State.NumFiles && State.NumFailingDownloads < NumThreads && State.LastDecompressError == null; BufferIdx = (BufferIdx + 1) % NumBytesReadBuffer.Length)
+			for (int BufferIdx = 0, NumFilesReportedRead = 0; NumFilesReportedRead < State.NumFiles && State.NumFailingOrIdleDownloads < NumThreads; BufferIdx = (BufferIdx + 1) % NumBytesReadBuffer.Length)
 			{
 				const int TickInterval = 100;
+				Thread.Sleep(TickInterval);
 
 				long NumBytesRead = Interlocked.Read(ref State.NumBytesRead);
+				long NumBytesTotal = Interlocked.Read(ref State.NumBytesTotal);
+				long NumBytesCached = Interlocked.Read(ref State.NumBytesCached);
+				long NumBytesNeed = NumBytesTotal - NumBytesCached;
 				float NumBytesPerSecond = (float)Math.Max(NumBytesRead - NumBytesReadBuffer[BufferIdx], 0) * 1000.0f / (NumBytesReadBuffer.Length * TickInterval);
 				NumFilesReportedRead = State.NumFilesRead;
-				Log.WriteStatus("Received {0}/{1} files ({2:0.0}/{3:0.0}mb; {4:0.00}mb/s; {5}%)...", NumFilesReportedRead, State.NumFiles, (NumBytesRead / (1024.0 * 1024.0)) + 0.0999999, (NumBytesTotal / (1024.0 * 1024.0)) + 0.0999999, (NumBytesPerSecond / (1024.0 * 1024.0)) + 0.0099, (NumBytesRead * 100) / NumBytesTotal);
-				NumBytesReadBuffer[BufferIdx] = NumBytesRead;
+				Log.WriteStatus("Received {0}/{1} files ({2:0.0}/{3:0.0}mb; {4:0.00}mb/s; {5}%; {6:0.0}mb from cache)...", NumFilesReportedRead, State.NumFiles, (NumBytesRead / (1024.0 * 1024.0)) + 0.0999999, (NumBytesNeed / (1024.0 * 1024.0)) + 0.0999999, (NumBytesPerSecond / (1024.0 * 1024.0)) + 0.0099, ((NumBytesRead + NumBytesCached) * 100) / NumBytesTotal, (NumBytesCached / (1024.0 * 1024.0)) + 0.0999999);
 
-				Thread.Sleep(TickInterval);
+				NumBytesReadBuffer[BufferIdx] = NumBytesRead;
 			}
 
 			// If we finished with an error, try to clean up and return
 			if(State.NumFilesRead < State.NumFiles)
 			{
-				DecompressionThread.Abort();
 				foreach(Thread WorkerThread in WorkerThreads)
 				{
 					WorkerThread.Abort();
 				}
-				Log.WriteError("{0}", (State.LastDecompressError != null)? State.LastDecompressError : State.LastDownloadError);
-				foreach(string FileName in DownloadFileNames.Values)
-				{
-					try { File.Delete(FileName); } catch (Exception) { }
-				}
 				return false;
 			}
 
 			// Join all the threads
-			DecompressionThread.Join();
 			foreach(Thread WorkerThread in WorkerThreads)
 			{
 				WorkerThread.Join();
@@ -687,76 +727,70 @@ namespace GitDependencies
 			return true;
 		}
 
-		static void DecompressWorker(string RootPath, ConcurrentQueue<DependencyPack> DecompressQueue, Dictionary<DependencyPack, string> DownloadFileNames, Dictionary<string, List<DependencyBlob>> PackToBlobs, Dictionary<string, List<DependencyFile>> BlobToFiles, AsyncDownloadState State)
+		static string GetPackCacheFile(string CachePath, DependencyPack Pack)
 		{
-			while(State.NumFilesRead < State.NumFiles)
+			if (CachePath == null)
 			{
-				// Remove the next file from the queue, or wait before polling again
-				DependencyPack NextPack;
-				if(!DecompressQueue.TryDequeue(out NextPack))
-				{
-					Thread.Sleep(100);
-					continue;
-				}
-
-				// Get the filename for the downloaded pack
-				string PackFileName = DownloadFileNames[NextPack];
-
-				// Extract all the files from this pack file to their final locations
-				foreach(DependencyBlob Blob in PackToBlobs[NextPack.Hash])
-				{
-					foreach(DependencyFile File in BlobToFiles[Blob.Hash])
-					{
-						string OutputFileName = Path.Combine(RootPath, File.Name);
-						try
-						{
-							ExtractBlob(PackFileName, Blob, OutputFileName);
-						}
-						catch(Exception Ex)
-						{
-							Interlocked.CompareExchange(ref State.LastDecompressError, String.Format("Failed to extract '{0}': {1}", OutputFileName, Ex.Message), null);
-							return;
-						}
-						Interlocked.Increment(ref State.NumFilesRead);
-					}
-				}
-
-				// Delete the pack file now that we're finished with it. Doesn't matter much if it fails.
-				try { System.IO.File.Delete(PackFileName); } catch(Exception) { }
+				return null;
+			}
+			string TargetDir = Path.Combine(CachePath, Pack.Hash.Substring(0, 2));
+			if (!Directory.Exists(TargetDir))
+			{
+				Directory.CreateDirectory(TargetDir);
 			}
+			return Path.Combine(TargetDir, Pack.Hash);
 		}
 
-		static void DownloadWorker(string RootPath, ConcurrentQueue<DependencyPack> DownloadQueue, ConcurrentQueue<DependencyPack> DecompressQueue, Dictionary<DependencyPack, string> DownloadFileNames, Dictionary<string, List<DependencyBlob>> PackToBlobs, Dictionary<string, List<DependencyFile>> BlobToFiles, AsyncDownloadState State, int MaxRetries, Uri Proxy)
+		static void DownloadWorker(string RootPath, string CachePath, ConcurrentQueue<DependencyPackInfo> DownloadQueue, Dictionary<string, List<DependencyBlob>> PackToBlobs, Dictionary<string, List<DependencyFile>> BlobToFiles, AsyncDownloadState State, int MaxRetries, Uri Proxy)
 		{
 			int Retries = 0;
 			while(State.NumFilesRead < State.NumFiles)
 			{
 				// Remove the next file from the download queue, or wait before polling again
-				DependencyPack NextPack;
-				if(!DownloadQueue.TryDequeue(out NextPack))
+				DependencyPackInfo NextPackInfo;
+				if (!DownloadQueue.TryDequeue(out NextPackInfo))
 				{
 					Thread.Sleep(100);
 					continue;
 				}
 
+				DependencyPack NextPack = NextPackInfo.Pack;
+
 				// Get the temporary file to download to
-				string PackFileName = DownloadFileNames[NextPack];
+				string CacheFileName = GetPackCacheFile(CachePath, NextPack);
 
 				// Format the URL for it
-				string Url = String.Format("http://cdn.unrealengine.com/dependencies/{0}/{1}", NextPack.RemotePath, NextPack.Hash);
+				string BaseUrl = NextPackInfo.Manifest.BaseUrl;
+				if (String.IsNullOrEmpty(BaseUrl))
+				{
+					BaseUrl = "http://cdn.unrealengine.com/dependencies";
+				}
+				string Url = String.Format("{0}/{1}/{2}", BaseUrl, NextPack.RemotePath, NextPack.Hash);
 
 				// Try to download the file
 				long RollbackSize = 0;
+				// Lock to correctly load the file from multiple processes.
+				Mutex FileLock = new Mutex(false, "ue-dependency-" + NextPack.Hash);
+				FileLock.WaitOne();
 				try
 				{
-					// Download the file and queue it for decompression
-					DownloadFileAndVerifyHash(Url, Proxy, PackFileName, NextPack.Hash, Size => { RollbackSize += Size; Interlocked.Add(ref State.NumBytesRead, Size); });
-					DecompressQueue.Enqueue(NextPack);
+					Dictionary<string, DependencyBlob> Blobs = PackBlobs(RootPath, PackToBlobs[NextPack.Hash], BlobToFiles);
+					if (TryUnpackFromCache(CacheFileName, NextPack.CompressedSize, Blobs))
+					{
+						Interlocked.Add(ref State.NumBytesCached, NextPack.CompressedSize);
+					}
+					else
+					{
+						// Download the file and queue it for decompression
+						long ReadSize = DownloadFileAndVerifyHash(Url, NextPackInfo.Manifest.IgnoreProxy ? null : Proxy, CacheFileName, NextPack, Blobs, Size => { RollbackSize += Size; Interlocked.Add(ref State.NumBytesRead, Size); });
+						Interlocked.Add(ref State.NumBytesTotal, ReadSize - NextPack.CompressedSize);
+					}
+					Interlocked.Add(ref State.NumFilesRead, Blobs.Count);
 
 					// If we were failing, decrement the number of failing threads
 					if(Retries > MaxRetries)
 					{
-						Interlocked.Decrement(ref State.NumFailingDownloads);
+						Interlocked.Decrement(ref State.NumFailingOrIdleDownloads);
 						Retries = 0;
 					}
 				}
@@ -764,19 +798,53 @@ namespace GitDependencies
 				{
 					// Rollback the byte count and add the file back into the download queue
 					Interlocked.Add(ref State.NumBytesRead, -RollbackSize);
-					DownloadQueue.Enqueue(NextPack);
-
+					DownloadQueue.Enqueue(NextPackInfo);
 					// If we've retried enough times already, set the error message. 
-					if(Retries++ == MaxRetries)
+					if (Retries++ == MaxRetries)
+					{
+						Interlocked.Increment(ref State.NumFailingOrIdleDownloads);
+						State.LastDownloadError = String.Format("Failed to download '{0}': {1} ({2})", Url, Ex.Message, Ex.GetType().Name);
+					}
+				}
+				finally
+				{
+					FileLock.ReleaseMutex();
+					FileLock.Dispose();
+				}
+			}
+			if (Retries < MaxRetries)
+			{
+				Interlocked.Increment(ref State.NumFailingOrIdleDownloads);
+			}
+		}
+
+		static bool TryUnpackFromCache(string CacheFileName, long CompressedSize, Dictionary<string, DependencyBlob> Blobs)
+		{
+			if (CacheFileName != null)
+			{
+				FileInfo PackFileInfo = new FileInfo(CacheFileName);
+				if (PackFileInfo.Exists && (PackFileInfo.Length == CompressedSize))
+				{
+					try
+					{
+						// Don't download file, if already exists.
+						using (GZipStream DecompressedStream = new GZipStream(PackFileInfo.OpenRead(), CompressionMode.Decompress, false))
+						{
+							ExtractBlobs(DecompressedStream, Blobs);
+							return true;
+						}
+					}
+					catch (InvalidCacheException)
 					{
-						Interlocked.Increment(ref State.NumFailingDownloads);
-						State.LastDownloadError = String.Format("Failed to download '{0}' to '{1}': {2} ({3})", Url, PackFileName, Ex.Message, Ex.GetType().Name);
+						try { File.Delete(CacheFileName); }
+						catch (Exception) { }
 					}
 				}
 			}
+			return false;
 		}
 
-		static void DownloadFileAndVerifyHash(string Url, Uri Proxy, string PackFileName, string ExpectedHash, NotifyReadDelegate NotifyRead)
+		static long DownloadFileAndVerifyHash(string Url, Uri Proxy, string CacheFileName, DependencyPack Pack, Dictionary<string, DependencyBlob> Blobs, NotifyReadDelegate NotifyRead)
 		{
 			// Create the web request
 			WebRequest Request = WebRequest.Create(Url);
@@ -790,26 +858,48 @@ namespace GitDependencies
 			}
 
 			// Get the response
-			using(WebResponse Response = Request.GetResponse())
+			using (WebResponse Response = Request.GetResponse())
 			{
-				// Download the file, decompressing and hashing it as we go
-				SHA1Managed Hasher = new SHA1Managed();
-				using(FileStream OutputStream = File.OpenWrite(PackFileName))
+				if (CacheFileName != null)
 				{
-					CryptoStream HashOutputStream = new CryptoStream(OutputStream, Hasher, CryptoStreamMode.Write);
-					using(NotifyReadStream InputStream = new NotifyReadStream(Response.GetResponseStream(), NotifyRead))
+					// Download the file, decompressing and hashing it as we go
+					SHA1Managed Hasher = new SHA1Managed();
+					long ReadSize;
+					using (FileStream OutputStream = File.OpenWrite(CacheFileName))
 					{
-						GZipStream DecompressedStream = new GZipStream(InputStream, CompressionMode.Decompress, true);
-						DecompressedStream.CopyTo(HashOutputStream);
+						OutputStream.SetLength(Pack.CompressedSize);
+						using (Stream InputStream = new ForkReadStream(new NotifyReadStream(Response.GetResponseStream(), NotifyRead), OutputStream))
+						{
+							GZipStream DecompressedStream = new GZipStream(InputStream, CompressionMode.Decompress, true);
+							CryptoStream HashInputStream = new CryptoStream(DecompressedStream, Hasher, CryptoStreamMode.Read);
+							ExtractBlobs(HashInputStream, Blobs);
+							// Read stream to end for compute hash and write to cache.
+							byte[] Buffer = new Byte[16384];
+							while (HashInputStream.Read(Buffer, 0, Buffer.Length) > 0) { }
+							HashInputStream.Close();
+							ReadSize = InputStream.Position;
+						}
+					}
+					// Check the hash was what we expected
+					string Hash = BitConverter.ToString(Hasher.Hash).ToLower().Replace("-", "");
+					if (Hash != Pack.Hash)
+					{
+						File.Delete(CacheFileName);
+						throw new InvalidDataException(String.Format("Incorrect hash for {0} - expected {1}, got {2}", Url, Pack.Hash, Hash));
 					}
-					HashOutputStream.FlushFinalBlock();
+					return ReadSize;
 				}
-
-				// Check the hash was what we expected
-				string Hash = BitConverter.ToString(Hasher.Hash).ToLower().Replace("-", "");
-				if(Hash != ExpectedHash)
+				else
 				{
-					throw new InvalidDataException(String.Format("Incorrect hash for {0} - expected {1}, got {2}", Url, ExpectedHash, Hash));
+					// Download the file, decompressing and hashing it as we go
+					SHA1Managed Hasher = new SHA1Managed();
+					using (Stream InputStream = new NotifyReadStream(Response.GetResponseStream(), NotifyRead))
+					{
+						GZipStream DecompressedStream = new GZipStream(InputStream, CompressionMode.Decompress, true);
+						CryptoStream HashInputStream = new CryptoStream(DecompressedStream, Hasher, CryptoStreamMode.Read);
+						ExtractBlobs(HashInputStream, Blobs);
+						return InputStream.Position;
+					}
 				}
 			}
 		}
@@ -827,43 +917,121 @@ namespace GitDependencies
 			return null;
 		}
 
-		static void ExtractBlob(string PackFileName, DependencyBlob Blob, string OutputFileName)
+		class BlobState
+		{
+			public String Path;
+			public Stream Stream;
+			public HashAlgorithm Hash;
+			public DependencyBlob Blob;
+
+			public BlobState(String Path, DependencyBlob Blob) 
+			{
+				this.Path = Path;
+				this.Blob = Blob;
+				this.Stream = null;
+				this.Hash = new SHA1Managed();
+			}
+		}
+
+		static Dictionary<string, DependencyBlob> PackBlobs(string RootPath, List<DependencyBlob> Blobs, Dictionary<string, List<DependencyFile>> BlobToFiles)
 		{
-			// Create the output folder
-			Directory.CreateDirectory(Path.GetDirectoryName(OutputFileName));
+			Dictionary<string, DependencyBlob> FileBlobs = new Dictionary<string, DependencyBlob>();
+			foreach (DependencyBlob Blob in Blobs)
+			{
+				foreach (DependencyFile File in BlobToFiles[Blob.Hash])
+				{
+					string OutputFileName = Path.Combine(RootPath, File.Name);
+					FileBlobs.Add(OutputFileName, Blob);
+				}
+			}
+			return FileBlobs;
+		}
 
-			// Copy the data to the output file
-			SHA1Managed Hasher = new SHA1Managed();
-			using(FileStream OutputStream = File.OpenWrite(OutputFileName + IncomingFileSuffix))
+		static void ExtractBlobs(Stream PackStream, Dictionary<string, DependencyBlob> Blobs)
+		{
+			List<BlobState> BlobStates = new List<BlobState>();
+			try
 			{
-				CryptoStream HashOutputStream = new CryptoStream(OutputStream, Hasher, CryptoStreamMode.Write);
-				using(FileStream InputStream = File.OpenRead(PackFileName))
+				// Prepare output files sorted by offset.
+				foreach (var Pair in Blobs)
 				{
-					// Seek to the right position
-					InputStream.Seek(Blob.PackOffset, SeekOrigin.Begin);
+					BlobStates.Add(new BlobState(Pair.Key, Pair.Value));
+				}
+				BlobStates.Sort((a, b) => a.Blob.PackOffset.CompareTo(b.Blob.PackOffset));
 
-					// Extract the data
-					byte[] Buffer = new byte[16384];
-					for(long RemainingSize = Blob.Size; RemainingSize > 0; )
+				// Create files from pack.
+				byte[] Buffer = new byte[16384];
+				long PackOffset = 0;
+				while (BlobStates.Count > 0)
+				{
+					int ReadSize;
+					try
+					{
+						ReadSize = PackStream.Read(Buffer, 0, Buffer.Length);
+					}
+					catch (Exception e)
+					{
+						throw new InvalidCacheException("Can't read from pack stream", e);
+					}
+					if (ReadSize == 0)
 					{
-						int ReadSize = InputStream.Read(Buffer, 0, (int)Math.Min(RemainingSize, (long)Buffer.Length));
-						if(ReadSize == 0) throw new InvalidDataException("Unexpected end of file");
-						HashOutputStream.Write(Buffer, 0, ReadSize);
-						RemainingSize -= ReadSize;
+						throw new InvalidCacheException("Unexpected end of file", null);
 					}
+					for (int i = 0; i < BlobStates.Count; ++i)
+					{
+						var State = BlobStates[i];
+						if (State.Blob.PackOffset > PackOffset + ReadSize) break;
+						if (State.Stream == null)
+						{
+							// Create the output folder
+							Directory.CreateDirectory(Path.GetDirectoryName(State.Path));
+							// Create file
+							State.Stream = File.OpenWrite(State.Path + IncomingFileSuffix);
+							State.Stream.SetLength(State.Blob.Size);
+							State.Hash.Initialize();
+						}
+						int BufferOffset = (int)Math.Max(0, State.Blob.PackOffset - PackOffset);
+						int BufferCount = (int)Math.Min(ReadSize, State.Blob.PackOffset - PackOffset + State.Blob.Size) - BufferOffset;
+						State.Stream.Write(Buffer, BufferOffset, BufferCount);
+						State.Hash.TransformBlock(Buffer, BufferOffset, BufferCount, Buffer, BufferOffset);
+						if (State.Blob.PackOffset + State.Blob.Size <= PackOffset + ReadSize) 
+						{
+							State.Hash.TransformFinalBlock(Buffer, 0, 0);
+							// Check the hash was what we expected
+							string ActualHash = BitConverter.ToString(State.Hash.Hash).ToLower().Replace("-", "");
+							if (ActualHash != State.Blob.Hash)
+							{
+								throw new InvalidCacheException("Incorrect hash value of file: " + State.Path, null);
+							}
+							State.Stream.Close();
+							State.Stream = null;
+							// Move the file to its final position
+							if (File.Exists(State.Path))
+							{
+								File.Delete(State.Path);
+							}
+							File.Move(State.Path + IncomingFileSuffix, State.Path);
+							// Remove file from queue
+							BlobStates.RemoveAt(i);
+							--i;
+						}
+					}
+					PackOffset += ReadSize;
 				}
-				HashOutputStream.FlushFinalBlock();
 			}
-
-			// Check the hash was what we expected
-			string Hash = BitConverter.ToString(Hasher.Hash).ToLower().Replace("-", "");
-			if(Hash != Blob.Hash)
+			finally 
 			{
-				throw new InvalidDataException("Incorrect hash value");
+				// Delete unfinished files.
+				foreach (var State in BlobStates)
+				{
+					if (State.Stream != null)
+					{
+						State.Stream.Close();
+						try { File.Delete(State.Path + IncomingFileSuffix); }
+						catch (Exception) { }
+					}
+				}
 			}
-
-			// Move the file to its final position
-			File.Move(OutputFileName + IncomingFileSuffix, OutputFileName);
 		}
 
 		public static bool ReadXmlObject<T>(string FileName, out T NewObject)
