Rename InvalidCacheException to CorruptPackFileException, since it isn't necessarily caused by an invalid cache entry. Other minor housekeeping.

diff --git a/Engine/Source/Programs/GitDependencies/Program.cs b/Engine/Source/Programs/GitDependencies/Program.cs
--- a/Engine/Source/Programs/GitDependencies/Program.cs
+++ b/Engine/Source/Programs/GitDependencies/Program.cs
@@ -48,9 +48,9 @@ namespace GitDependencies
 			}
 		}
 
-		class InvalidCacheException : Exception
+		class CorruptPackFileException : Exception
 		{
-			public InvalidCacheException(string Message, Exception InnerException)
+			public CorruptPackFileException(string Message, Exception InnerException)
 				: base(Message, InnerException)
 			{
 			}
@@ -75,15 +75,9 @@ namespace GitDependencies
 			bool bHelp = ParseSwitch(ArgsList, "-help");
 			string RootPath = ParseParameter(ArgsList, "-root=", Path.GetFullPath(Path.Combine(Path.GetDirectoryName(Assembly.GetExecutingAssembly().Location), "../../..")));
 
-			//
-			// Cache path
-			// 
-			// * If specified flag --no-cache - do not use the cache;
-			// * If specified flag --cache=<PATH> - use the path <PATH>;
-			// * If specified enviroment variable UE4_GITDEPS, then use path from it;
-			// * If there is a .git directory above from --root=<PATH>, then use the path of the form .git/ue4-gitdeps/;
-			// * In other case - do not use the cache.
-			//
+			// Parse the cache path. A specific path can be set using -catch=<PATH> or the UE4_GITDEPS environment variable, otherwise we look for a parent .git directory
+			// and use a sub-folder of that. Users which download the source through a zip file (and won't have a .git directory) are unlikely to benefit from caching, as
+			// they only need to download dependencies once.
 			string CachePath = null;
 			if (!ParseSwitch(ArgsList, "-no-cache"))
 			{
@@ -834,10 +828,9 @@ namespace GitDependencies
 							return true;
 						}
 					}
-					catch (InvalidCacheException)
+					catch (CorruptPackFileException)
 					{
-						try { File.Delete(CacheFileName); }
-						catch (Exception) { }
+						SafeDeleteFileQuiet(CacheFileName);
 					}
 				}
 			}
@@ -971,11 +964,11 @@ namespace GitDependencies
 					}
 					catch (Exception e)
 					{
-						throw new InvalidCacheException("Can't read from pack stream", e);
+						throw new CorruptPackFileException("Can't read from pack stream", e);
 					}
 					if (ReadSize == 0)
 					{
-						throw new InvalidCacheException("Unexpected end of file", null);
+						throw new CorruptPackFileException("Unexpected end of file", null);
 					}
 					for (int i = 0; i < BlobStates.Count; ++i)
 					{
@@ -1001,7 +994,7 @@ namespace GitDependencies
 							string ActualHash = BitConverter.ToString(State.Hash.Hash).ToLower().Replace("-", "");
 							if (ActualHash != State.Blob.Hash)
 							{
-								throw new InvalidCacheException("Incorrect hash value of file: " + State.Path, null);
+								throw new CorruptPackFileException("Incorrect hash value of file: " + State.Path, null);
 							}
 							State.Stream.Close();
 							State.Stream = null;
@@ -1027,14 +1020,13 @@ namespace GitDependencies
 					if (State.Stream != null)
 					{
 						State.Stream.Close();
-						try { File.Delete(State.Path + IncomingFileSuffix); }
-						catch (Exception) { }
+						SafeDeleteFileQuiet(State.Path + IncomingFileSuffix);
 					}
 				}
 			}
 		}
 
-		public static bool ReadXmlObject<T>(string FileName, out T NewObject)
+		static bool ReadXmlObject<T>(string FileName, out T NewObject)
 		{
 			try
 			{
@@ -1053,7 +1045,7 @@ namespace GitDependencies
 			}
 		}
 
-		public static bool WriteXmlObject<T>(string FileName, T XmlObject)
+		static bool WriteXmlObject<T>(string FileName, T XmlObject)
 		{
 			try
 			{
@@ -1071,7 +1063,7 @@ namespace GitDependencies
 			}
 		}
 
-		public static bool WriteWorkingManifest(string FileName, string TemporaryFileName, WorkingManifest Manifest)
+		static bool WriteWorkingManifest(string FileName, string TemporaryFileName, WorkingManifest Manifest)
 		{
 			if(!WriteXmlObject(TemporaryFileName, Manifest))
 			{
@@ -1092,7 +1084,7 @@ namespace GitDependencies
 			return true;
 		}
 
-		public static bool SafeModifyFileAttributes(string FileName, FileAttributes AddAttributes, FileAttributes RemoveAttributes)
+		static bool SafeModifyFileAttributes(string FileName, FileAttributes AddAttributes, FileAttributes RemoveAttributes)
 		{
 			try
 			{
@@ -1106,7 +1098,7 @@ namespace GitDependencies
 			}
 		}
 
-		public static bool SafeCreateDirectory(string DirectoryName)
+		static bool SafeCreateDirectory(string DirectoryName)
 		{
 			try
 			{
@@ -1120,7 +1112,7 @@ namespace GitDependencies
 			}
 		}
 
-		public static bool SafeDeleteFile(string FileName)
+		static bool SafeDeleteFile(string FileName)
 		{
 			try
 			{
@@ -1134,7 +1126,20 @@ namespace GitDependencies
 			}
 		}
 
-		public static bool SafeMoveFile(string SourceFileName, string TargetFileName)
+		static bool SafeDeleteFileQuiet(string FileName)
+		{
+			try
+			{
+				File.Delete(FileName);
+				return true;
+			}
+			catch(IOException)
+			{
+				return false;
+			}
+		}
+
+		static bool SafeMoveFile(string SourceFileName, string TargetFileName)
 		{
 			try
 			{
@@ -1148,7 +1153,7 @@ namespace GitDependencies
 			}
 		}
 
-		public static string ComputeHashForFile(string FileName)
+		static string ComputeHashForFile(string FileName)
 		{
 			using(FileStream InputStream = File.OpenRead(FileName))
 			{
