Use a lockless method of getting things into the cache. Relies on file system moves being atomic; should allow using a cache on a shared or network drive.

diff --git a/Engine/Source/Programs/GitDependencies/Program.cs b/Engine/Source/Programs/GitDependencies/Program.cs
--- a/Engine/Source/Programs/GitDependencies/Program.cs
+++ b/Engine/Source/Programs/GitDependencies/Program.cs
@@ -58,7 +58,9 @@ namespace GitDependencies
 
 		const string IncomingFileSuffix = ".incoming";
 		const string TempManifestExtension = ".tmp";
-		
+
+		static readonly string InstanceSuffix = Guid.NewGuid().ToString().Replace("-", "");
+
 		static int Main(string[] Args)
 		{
 			// Build the argument list. Remove any double-hyphens from the start of arguments for conformity with other Epic tools.
@@ -694,10 +696,9 @@ namespace GitDependencies
 				long NumBytesRead = Interlocked.Read(ref State.NumBytesRead);
 				long NumBytesTotal = Interlocked.Read(ref State.NumBytesTotal);
 				long NumBytesCached = Interlocked.Read(ref State.NumBytesCached);
-				long NumBytesNeed = NumBytesTotal - NumBytesCached;
 				float NumBytesPerSecond = (float)Math.Max(NumBytesRead - NumBytesReadBuffer[BufferIdx], 0) * 1000.0f / (NumBytesReadBuffer.Length * TickInterval);
 				NumFilesReportedRead = State.NumFilesRead;
-				Log.WriteStatus("Received {0}/{1} files ({2:0.0}/{3:0.0}mb; {4:0.00}mb/s; {5}%; {6:0.0}mb from cache)...", NumFilesReportedRead, State.NumFiles, (NumBytesRead / (1024.0 * 1024.0)) + 0.0999999, (NumBytesNeed / (1024.0 * 1024.0)) + 0.0999999, (NumBytesPerSecond / (1024.0 * 1024.0)) + 0.0099, ((NumBytesRead + NumBytesCached) * 100) / NumBytesTotal, (NumBytesCached / (1024.0 * 1024.0)) + 0.0999999);
+				Log.WriteStatus("Received {0}/{1} files ({2:0.0}/{3:0.0}mb; {4:0.00}mb/s; {5}%; {6:0.0}mb from cache)...", NumFilesReportedRead, State.NumFiles, (NumBytesRead / (1024.0 * 1024.0)) + 0.0999999, (NumBytesTotal / (1024.0 * 1024.0)) + 0.0999999, (NumBytesPerSecond / (1024.0 * 1024.0)) + 0.0099, ((NumBytesRead + NumBytesCached) * 100) / (NumBytesTotal + NumBytesCached), (NumBytesCached / (1024.0 * 1024.0)) + 0.0999999);
 
 				NumBytesReadBuffer[BufferIdx] = NumBytesRead;
 			}
@@ -763,9 +764,6 @@ namespace GitDependencies
 
 				// Try to download the file
 				long RollbackSize = 0;
-				// Lock to correctly load the file from multiple processes.
-				Mutex FileLock = new Mutex(false, "ue-dependency-" + NextPack.Hash);
-				FileLock.WaitOne();
 				try
 				{
 					Dictionary<string, DependencyBlob> Blobs = PackBlobs(RootPath, PackToBlobs[NextPack.Hash], BlobToFiles);
@@ -775,10 +773,11 @@ namespace GitDependencies
 					}
 					else
 					{
-						// Download the file and queue it for decompression
-						long ReadSize = DownloadFileAndVerifyHash(Url, NextPackInfo.Manifest.IgnoreProxy ? null : Proxy, CacheFileName, NextPack, Blobs, Size => { RollbackSize += Size; Interlocked.Add(ref State.NumBytesRead, Size); });
-						Interlocked.Add(ref State.NumBytesTotal, ReadSize - NextPack.CompressedSize);
+						DownloadAndExtractFiles(Url, NextPackInfo.Manifest.IgnoreProxy ? null : Proxy, CacheFileName, NextPack, Blobs, Size => { RollbackSize += Size; Interlocked.Add(ref State.NumBytesRead, Size); });
 					}
+
+					// Update the stats
+					Interlocked.Add(ref State.NumBytesTotal, RollbackSize - NextPack.CompressedSize);
 					Interlocked.Add(ref State.NumFilesRead, Blobs.Count);
 
 					// If we were failing, decrement the number of failing threads
@@ -793,6 +792,7 @@ namespace GitDependencies
 					// Rollback the byte count and add the file back into the download queue
 					Interlocked.Add(ref State.NumBytesRead, -RollbackSize);
 					DownloadQueue.Enqueue(NextPackInfo);
+
 					// If we've retried enough times already, set the error message. 
 					if (Retries++ == MaxRetries)
 					{
@@ -800,11 +800,6 @@ namespace GitDependencies
 						State.LastDownloadError = String.Format("Failed to download '{0}': {1} ({2})", Url, Ex.Message, Ex.GetType().Name);
 					}
 				}
-				finally
-				{
-					FileLock.ReleaseMutex();
-					FileLock.Dispose();
-				}
 			}
 			if (Retries < MaxRetries)
 			{
@@ -837,7 +832,7 @@ namespace GitDependencies
 			return false;
 		}
 
-		static long DownloadFileAndVerifyHash(string Url, Uri Proxy, string CacheFileName, DependencyPack Pack, Dictionary<string, DependencyBlob> Blobs, NotifyReadDelegate NotifyRead)
+		static void DownloadAndExtractFiles(string Url, Uri Proxy, string CacheFileName, DependencyPack Pack, Dictionary<string, DependencyBlob> Blobs, NotifyReadDelegate NotifyRead)
 		{
 			// Create the web request
 			WebRequest Request = WebRequest.Create(Url);
@@ -850,48 +845,18 @@ namespace GitDependencies
 				Request.Proxy = new WebProxy(Proxy, true, null, MakeCredentialsFromUri(Proxy));
 			}
 
-			// Get the response
+			// Read the response and extract the files
 			using (WebResponse Response = Request.GetResponse())
 			{
-				if (CacheFileName != null)
+				using (Stream ResponseStream = new NotifyReadStream(Response.GetResponseStream(), NotifyRead))
 				{
-					// Download the file, decompressing and hashing it as we go
-					SHA1Managed Hasher = new SHA1Managed();
-					long ReadSize;
-					using (FileStream OutputStream = File.OpenWrite(CacheFileName))
+					if(CacheFileName == null)
 					{
-						OutputStream.SetLength(Pack.CompressedSize);
-						using (Stream InputStream = new ForkReadStream(new NotifyReadStream(Response.GetResponseStream(), NotifyRead), OutputStream))
-						{
-							GZipStream DecompressedStream = new GZipStream(InputStream, CompressionMode.Decompress, true);
-							CryptoStream HashInputStream = new CryptoStream(DecompressedStream, Hasher, CryptoStreamMode.Read);
-							ExtractBlobs(HashInputStream, Blobs);
-							// Read stream to end for compute hash and write to cache.
-							byte[] Buffer = new Byte[16384];
-							while (HashInputStream.Read(Buffer, 0, Buffer.Length) > 0) { }
-							HashInputStream.Close();
-							ReadSize = InputStream.Position;
-						}
-					}
-					// Check the hash was what we expected
-					string Hash = BitConverter.ToString(Hasher.Hash).ToLower().Replace("-", "");
-					if (Hash != Pack.Hash)
-					{
-						File.Delete(CacheFileName);
-						throw new InvalidDataException(String.Format("Incorrect hash for {0} - expected {1}, got {2}", Url, Pack.Hash, Hash));
+						ExtractFiles(ResponseStream, Blobs);
 					}
-					return ReadSize;
-				}
-				else
-				{
-					// Download the file, decompressing and hashing it as we go
-					SHA1Managed Hasher = new SHA1Managed();
-					using (Stream InputStream = new NotifyReadStream(Response.GetResponseStream(), NotifyRead))
+					else
 					{
-						GZipStream DecompressedStream = new GZipStream(InputStream, CompressionMode.Decompress, true);
-						CryptoStream HashInputStream = new CryptoStream(DecompressedStream, Hasher, CryptoStreamMode.Read);
-						ExtractBlobs(HashInputStream, Blobs);
-						return InputStream.Position;
+						ExtractFilesThroughCache(ResponseStream, CacheFileName, Pack.Hash, Blobs);
 					}
 				}
 			}
@@ -910,6 +875,57 @@ namespace GitDependencies
 			return null;
 		}
 
+		static void ExtractFiles(Stream InputStream, Dictionary<string, DependencyBlob> Blobs)
+		{
+			// Create a decompression stream around the raw input stream
+			GZipStream DecompressedStream = new GZipStream(InputStream, CompressionMode.Decompress, true);
+			ExtractBlobs(DecompressedStream, Blobs);
+		}
+
+		static void ExtractFilesThroughCache(Stream InputStream, string FileName, string ExpectedHash, Dictionary<string, DependencyBlob> Blobs)
+		{
+			// Extract files from a pack file while writing to the cache file at the same time
+			string IncomingFileName = String.Format("{0}-{1}{2}", FileName, InstanceSuffix, IncomingFileSuffix);
+			try
+			{
+				// Make sure the directory exists
+				Directory.CreateDirectory(Path.GetDirectoryName(IncomingFileName));
+
+				// Hash the uncompressed data as we go
+				SHA1Managed Hasher = new SHA1Managed();
+				using(FileStream CacheStream = File.Open(IncomingFileName, FileMode.Create, FileAccess.Write, FileShare.None))
+				{
+					ForkReadStream ForkedInputStream = new ForkReadStream(InputStream, CacheStream);
+					GZipStream DecompressedStream = new GZipStream(ForkedInputStream, CompressionMode.Decompress, true);
+					using(CryptoStream DecompressedHashedStream = new CryptoStream(DecompressedStream, Hasher, CryptoStreamMode.Read))
+					{
+						ExtractBlobs(DecompressedHashedStream, Blobs);
+						ReadToEnd(DecompressedHashedStream);
+					}
+				}
+
+				// Check the hash was what we expected, and move it into the cache if it is.
+				string Hash = BitConverter.ToString(Hasher.Hash).ToLower().Replace("-", "");
+				if (Hash != ExpectedHash)
+				{
+					throw new CorruptPackFileException(String.Format("Incorrect hash for pack - expected {0}, got {1}", ExpectedHash, Hash), null);
+				}
+
+				// Move the new cache file into place
+				SafeMoveFileQuiet(IncomingFileName, FileName);
+			}
+			finally
+			{
+				SafeDeleteFileQuiet(IncomingFileName);
+			}
+		}
+
+		static void ReadToEnd(Stream InputStream)
+		{
+			byte[] Buffer = new byte[16384];
+			while (InputStream.Read(Buffer, 0, Buffer.Length) != 0) { }
+		}
+
 		class BlobState
 		{
 			public String Path;
@@ -998,12 +1014,11 @@ namespace GitDependencies
 							}
 							State.Stream.Close();
 							State.Stream = null;
+							
 							// Move the file to its final position
-							if (File.Exists(State.Path))
-							{
-								File.Delete(State.Path);
-							}
+							File.Delete(State.Path);
 							File.Move(State.Path + IncomingFileSuffix, State.Path);
+
 							// Remove file from queue
 							BlobStates.RemoveAt(i);
 							--i;
@@ -1153,6 +1168,19 @@ namespace GitDependencies
 			}
 		}
 
+		static bool SafeMoveFileQuiet(string SourceFileName, string TargetFileName)
+		{
+			try
+			{
+				File.Move(SourceFileName, TargetFileName);
+				return true;
+			}
+			catch(IOException)
+			{
+				return false;
+			}
+		}
+
 		static string ComputeHashForFile(string FileName)
 		{
 			using(FileStream InputStream = File.OpenRead(FileName))
