Fix ugly rounding errors in GitDependencies status message (always shows 0.1 instead of 0.0). Calculate it using integers instead.

diff --git a/Engine/Source/Programs/GitDependencies/Program.cs b/Engine/Source/Programs/GitDependencies/Program.cs
--- a/Engine/Source/Programs/GitDependencies/Program.cs
+++ b/Engine/Source/Programs/GitDependencies/Program.cs
@@ -696,9 +696,9 @@ namespace GitDependencies
 				long NumBytesRead = Interlocked.Read(ref State.NumBytesRead);
 				long NumBytesTotal = Interlocked.Read(ref State.NumBytesTotal);
 				long NumBytesCached = Interlocked.Read(ref State.NumBytesCached);
-				float NumBytesPerSecond = (float)Math.Max(NumBytesRead - NumBytesReadBuffer[BufferIdx], 0) * 1000.0f / (NumBytesReadBuffer.Length * TickInterval);
+				long NumBytesPerSecond = (long)Math.Ceiling((float)Math.Max(NumBytesRead - NumBytesReadBuffer[BufferIdx], 0) * 1000.0f / (NumBytesReadBuffer.Length * TickInterval));
 				NumFilesReportedRead = State.NumFilesRead;
-				Log.WriteStatus("Received {0}/{1} files ({2:0.0}/{3:0.0}mb; {4:0.00}mb/s; {5}%; {6:0.0}mb from cache)...", NumFilesReportedRead, State.NumFiles, (NumBytesRead / (1024.0 * 1024.0)) + 0.0999999, (NumBytesTotal / (1024.0 * 1024.0)) + 0.0999999, (NumBytesPerSecond / (1024.0 * 1024.0)) + 0.0099, ((NumBytesRead + NumBytesCached) * 100) / (NumBytesTotal + NumBytesCached), (NumBytesCached / (1024.0 * 1024.0)) + 0.0999999);
+				Log.WriteStatus("Received {0}/{1} files ({2}/{3}mb; {4}mb/s; {5}%; {6}mb from cache)...", NumFilesReportedRead, State.NumFiles, FormatMegabytes(NumBytesRead, 1), FormatMegabytes(NumBytesTotal, 1), FormatMegabytes(NumBytesPerSecond, 2), ((NumBytesRead + NumBytesCached) * 100) / (NumBytesTotal + NumBytesCached), FormatMegabytes(NumBytesCached, 1));
 
 				NumBytesReadBuffer[BufferIdx] = NumBytesRead;
 			}
@@ -722,6 +722,14 @@ namespace GitDependencies
 			return true;
 		}
 
+		static string FormatMegabytes(long Value, int NumDecimalPlaces)
+		{
+			int Multiplier = (int)Math.Pow(10.0, NumDecimalPlaces);
+			long FormatValue = ((Value * Multiplier) + (1024 * 1024) - 1) / (1024 * 1024);
+			string Result = String.Format("{0}.{1:D" + NumDecimalPlaces.ToString() + "}", FormatValue / Multiplier, FormatValue % Multiplier);
+			return Result;
+		}
+
 		static string GetPackCacheFile(string CachePath, DependencyPack Pack)
 		{
 			if (CachePath == null)
