Add separate classes for incoming pack/file data and calculate everything up-front, rather than passing through all the worker threads. Put the default base URL in the dependency manifest to prevent needing to special-case it.

diff --git a/Engine/Source/Programs/GitDependencies/DependencyManifest.cs b/Engine/Source/Programs/GitDependencies/DependencyManifest.cs
--- a/Engine/Source/Programs/GitDependencies/DependencyManifest.cs
+++ b/Engine/Source/Programs/GitDependencies/DependencyManifest.cs
@@ -59,7 +59,7 @@ namespace GitDependencies
 	public class DependencyManifest
 	{
 		[XmlAttribute]
-		public string BaseUrl;
+		public string BaseUrl = "http://cdn.unrealengine.com/dependencies";
 
 		[XmlAttribute]
 		public bool IgnoreProxy;
diff --git a/Engine/Source/Programs/GitDependencies/Program.cs b/Engine/Source/Programs/GitDependencies/Program.cs
--- a/Engine/Source/Programs/GitDependencies/Program.cs
+++ b/Engine/Source/Programs/GitDependencies/Program.cs
@@ -37,6 +37,24 @@ namespace GitDependencies
 			Force,
 		}
 
+		class IncomingPack
+		{
+			public string Url;
+			public Uri Proxy;
+			public string Hash;
+			public string CacheFileName;
+			public IncomingFile[] Files;
+			public long CompressedSize;
+		}
+
+		class IncomingFile
+		{
+			public string Name;
+			public string Hash;
+			public long MinPackOffset;
+			public long MaxPackOffset;
+		}
+
 		struct DependencyPackInfo
 		{
 			public DependencyManifest Manifest;
@@ -672,17 +690,30 @@ namespace GitDependencies
 			// Find all the required packs
 			DependencyPackInfo[] RequiredPacks = Packs.Where(x => PackToBlobs.ContainsKey(x.Pack.Hash)).ToArray();
 
+			// Create the download queue
+			ConcurrentQueue<IncomingPack> DownloadQueue = new ConcurrentQueue<IncomingPack>();
+			foreach(DependencyPackInfo RequiredPack in RequiredPacks)
+			{
+				IncomingPack Pack = new IncomingPack();
+				Pack.Url = String.Format("{0}/{1}/{2}", RequiredPack.Manifest.BaseUrl, RequiredPack.Pack.RemotePath, RequiredPack.Pack.Hash);
+				Pack.Proxy = RequiredPack.Manifest.IgnoreProxy? null : Proxy;
+				Pack.Hash = RequiredPack.Pack.Hash;
+				Pack.CacheFileName = (CachePath == null)? null : Path.Combine(CachePath, RequiredPack.Pack.Hash.Substring(0, 2), RequiredPack.Pack.Hash);
+				Pack.Files = GetIncomingFilesForPack(RootPath, RequiredPack.Pack, PackToBlobs, BlobToFiles);
+				Pack.CompressedSize = RequiredPack.Pack.CompressedSize;
+				DownloadQueue.Enqueue(Pack);
+			}
+
 			// Setup the async state
 			AsyncDownloadState State = new AsyncDownloadState();
 			State.NumFiles = RequiredFiles.Count();
 			State.NumBytesTotal = RequiredPacks.Sum(x => x.Pack.CompressedSize);
-			ConcurrentQueue<DependencyPackInfo> DownloadQueue = new ConcurrentQueue<DependencyPackInfo>(RequiredPacks);
 
 			// Create all the worker threads
 			Thread[] WorkerThreads = new Thread[NumThreads];
 			for(int Idx = 0; Idx < NumThreads; Idx++)
 			{
-				WorkerThreads[Idx] = new Thread(x => DownloadWorker(RootPath, CachePath, DownloadQueue, PackToBlobs, BlobToFiles, State, MaxRetries, Proxy));
+				WorkerThreads[Idx] = new Thread(x => DownloadWorker(DownloadQueue, State, MaxRetries));
 				WorkerThreads[Idx].Start();
 			}
 
@@ -730,63 +761,54 @@ namespace GitDependencies
 			return Result;
 		}
 
-		static string GetPackCacheFile(string CachePath, DependencyPack Pack)
+		static IncomingFile[] GetIncomingFilesForPack(string RootPath, DependencyPack RequiredPack, Dictionary<string, List<DependencyBlob>> PackToBlobs, Dictionary<string, List<DependencyFile>> BlobToFiles)
 		{
-			if (CachePath == null)
-			{
-				return null;
-			}
-			string TargetDir = Path.Combine(CachePath, Pack.Hash.Substring(0, 2));
-			if (!Directory.Exists(TargetDir))
+			List<IncomingFile> Files = new List<IncomingFile>();
+			foreach(DependencyBlob RequiredBlob in PackToBlobs[RequiredPack.Hash])
 			{
-				Directory.CreateDirectory(TargetDir);
+				foreach(DependencyFile RequiredFile in BlobToFiles[RequiredBlob.Hash])
+				{
+					IncomingFile File = new IncomingFile();
+					File.Name = Path.Combine(RootPath, RequiredFile.Name);
+					File.Hash = RequiredBlob.Hash;
+					File.MinPackOffset = RequiredBlob.PackOffset;
+					File.MaxPackOffset = RequiredBlob.PackOffset + RequiredBlob.Size;
+					Files.Add(File);
+				}
 			}
-			return Path.Combine(TargetDir, Pack.Hash);
+			return Files.OrderBy(x => x.MinPackOffset).ToArray();
 		}
 
-		static void DownloadWorker(string RootPath, string CachePath, ConcurrentQueue<DependencyPackInfo> DownloadQueue, Dictionary<string, List<DependencyBlob>> PackToBlobs, Dictionary<string, List<DependencyFile>> BlobToFiles, AsyncDownloadState State, int MaxRetries, Uri Proxy)
+		static void DownloadWorker(ConcurrentQueue<IncomingPack> DownloadQueue, AsyncDownloadState State, int MaxRetries)
 		{
 			int Retries = 0;
 			while(State.NumFilesRead < State.NumFiles)
 			{
 				// Remove the next file from the download queue, or wait before polling again
-				DependencyPackInfo NextPackInfo;
-				if (!DownloadQueue.TryDequeue(out NextPackInfo))
+				IncomingPack NextPack;
+				if (!DownloadQueue.TryDequeue(out NextPack))
 				{
 					Thread.Sleep(100);
 					continue;
 				}
 
-				DependencyPack NextPack = NextPackInfo.Pack;
-
-				// Get the temporary file to download to
-				string CacheFileName = GetPackCacheFile(CachePath, NextPack);
-
-				// Format the URL for it
-				string BaseUrl = NextPackInfo.Manifest.BaseUrl;
-				if (String.IsNullOrEmpty(BaseUrl))
-				{
-					BaseUrl = "http://cdn.unrealengine.com/dependencies";
-				}
-				string Url = String.Format("{0}/{1}/{2}", BaseUrl, NextPack.RemotePath, NextPack.Hash);
-
 				// Try to download the file
 				long RollbackSize = 0;
 				try
 				{
-					Dictionary<string, DependencyBlob> Blobs = PackBlobs(RootPath, PackToBlobs[NextPack.Hash], BlobToFiles);
-					if (TryUnpackFromCache(CacheFileName, NextPack.CompressedSize, Blobs))
+					// Download the pack file or extract it from the cache
+					if (TryUnpackFromCache(NextPack.CacheFileName, NextPack.CompressedSize, NextPack.Files))
 					{
 						Interlocked.Add(ref State.NumBytesCached, NextPack.CompressedSize);
 					}
 					else
 					{
-						DownloadAndExtractFiles(Url, NextPackInfo.Manifest.IgnoreProxy ? null : Proxy, CacheFileName, NextPack, Blobs, Size => { RollbackSize += Size; Interlocked.Add(ref State.NumBytesRead, Size); });
+						DownloadAndExtractFiles(NextPack.Url, NextPack.Proxy, NextPack.CacheFileName, NextPack.Hash, NextPack.Files, Size => { RollbackSize += Size; Interlocked.Add(ref State.NumBytesRead, Size); });
 					}
 
 					// Update the stats
 					Interlocked.Add(ref State.NumBytesTotal, RollbackSize - NextPack.CompressedSize);
-					Interlocked.Add(ref State.NumFilesRead, Blobs.Count);
+					Interlocked.Add(ref State.NumFilesRead, NextPack.Files.Length);
 
 					// If we were failing, decrement the number of failing threads
 					if(Retries > MaxRetries)
@@ -799,13 +821,13 @@ namespace GitDependencies
 				{
 					// Rollback the byte count and add the file back into the download queue
 					Interlocked.Add(ref State.NumBytesRead, -RollbackSize);
-					DownloadQueue.Enqueue(NextPackInfo);
+					DownloadQueue.Enqueue(NextPack);
 
 					// If we've retried enough times already, set the error message. 
 					if (Retries++ == MaxRetries)
 					{
 						Interlocked.Increment(ref State.NumFailingOrIdleDownloads);
-						State.LastDownloadError = String.Format("Failed to download '{0}': {1} ({2})", Url, Ex.Message, Ex.GetType().Name);
+						State.LastDownloadError = String.Format("Failed to download '{0}': {1} ({2})", NextPack.Url, Ex.Message, Ex.GetType().Name);
 					}
 				}
 			}
@@ -815,7 +837,7 @@ namespace GitDependencies
 			}
 		}
 
-		static bool TryUnpackFromCache(string CacheFileName, long CompressedSize, Dictionary<string, DependencyBlob> Blobs)
+		static bool TryUnpackFromCache(string CacheFileName, long CompressedSize, IncomingFile[] Files)
 		{
 			if (CacheFileName != null)
 			{
@@ -827,7 +849,7 @@ namespace GitDependencies
 						// Don't download file, if already exists.
 						using (GZipStream DecompressedStream = new GZipStream(PackFileInfo.OpenRead(), CompressionMode.Decompress, false))
 						{
-							ExtractBlobs(DecompressedStream, Blobs);
+							ExtractFilesFromRawStream(DecompressedStream, Files);
 							return true;
 						}
 					}
@@ -840,7 +862,7 @@ namespace GitDependencies
 			return false;
 		}
 
-		static void DownloadAndExtractFiles(string Url, Uri Proxy, string CacheFileName, DependencyPack Pack, Dictionary<string, DependencyBlob> Blobs, NotifyReadDelegate NotifyRead)
+		static void DownloadAndExtractFiles(string Url, Uri Proxy, string CacheFileName, string ExpectedHash, IncomingFile[] Files, NotifyReadDelegate NotifyRead)
 		{
 			// Create the web request
 			WebRequest Request = WebRequest.Create(Url);
@@ -860,11 +882,11 @@ namespace GitDependencies
 				{
 					if(CacheFileName == null)
 					{
-						ExtractFiles(ResponseStream, Blobs);
+						ExtractFiles(ResponseStream, Files);
 					}
 					else
 					{
-						ExtractFilesThroughCache(ResponseStream, CacheFileName, Pack.Hash, Blobs);
+						ExtractFilesThroughCache(ResponseStream, CacheFileName, ExpectedHash, Files);
 					}
 				}
 			}
@@ -872,6 +894,7 @@ namespace GitDependencies
 
 		static NetworkCredential MakeCredentialsFromUri(Uri Address)
 		{
+			// Check if the URI has a login:password prefix, and convert it to a NetworkCredential object if it has. HttpRequest just ignores it.
 			if(!String.IsNullOrEmpty(Address.UserInfo))
 			{
 				int Index = Address.UserInfo.IndexOf(':');
@@ -883,14 +906,14 @@ namespace GitDependencies
 			return null;
 		}
 
-		static void ExtractFiles(Stream InputStream, Dictionary<string, DependencyBlob> Blobs)
+		static void ExtractFiles(Stream InputStream, IncomingFile[] Files)
 		{
 			// Create a decompression stream around the raw input stream
 			GZipStream DecompressedStream = new GZipStream(InputStream, CompressionMode.Decompress, true);
-			ExtractBlobs(DecompressedStream, Blobs);
+			ExtractFilesFromRawStream(DecompressedStream, Files);
 		}
 
-		static void ExtractFilesThroughCache(Stream InputStream, string FileName, string ExpectedHash, Dictionary<string, DependencyBlob> Blobs)
+		static void ExtractFilesThroughCache(Stream InputStream, string FileName, string ExpectedHash, IncomingFile[] Files)
 		{
 			// Extract files from a pack file while writing to the cache file at the same time
 			string IncomingFileName = String.Format("{0}-{1}{2}", FileName, InstanceSuffix, IncomingFileSuffix);
@@ -907,7 +930,7 @@ namespace GitDependencies
 					GZipStream DecompressedStream = new GZipStream(ForkedInputStream, CompressionMode.Decompress, true);
 					using(CryptoStream DecompressedHashedStream = new CryptoStream(DecompressedStream, Hasher, CryptoStreamMode.Read))
 					{
-						ExtractBlobs(DecompressedHashedStream, Blobs);
+						ExtractFilesFromRawStream(DecompressedHashedStream, Files);
 						ReadToEnd(DecompressedHashedStream);
 					}
 				}
@@ -934,102 +957,69 @@ namespace GitDependencies
 			while (InputStream.Read(Buffer, 0, Buffer.Length) != 0) { }
 		}
 
-		class BlobState
+		static void ExtractFilesFromRawStream(Stream RawStream, IncomingFile[] Files)
 		{
-			public String Path;
-			public Stream Stream;
-			public HashAlgorithm Hash;
-			public DependencyBlob Blob;
-
-			public BlobState(String Path, DependencyBlob Blob) 
-			{
-				this.Path = Path;
-				this.Blob = Blob;
-				this.Stream = null;
-				this.Hash = new SHA1Managed();
-			}
-		}
-
-		static Dictionary<string, DependencyBlob> PackBlobs(string RootPath, List<DependencyBlob> Blobs, Dictionary<string, List<DependencyFile>> BlobToFiles)
-		{
-			Dictionary<string, DependencyBlob> FileBlobs = new Dictionary<string, DependencyBlob>();
-			foreach (DependencyBlob Blob in Blobs)
-			{
-				foreach (DependencyFile File in BlobToFiles[Blob.Hash])
-				{
-					string OutputFileName = Path.Combine(RootPath, File.Name);
-					FileBlobs.Add(OutputFileName, Blob);
-				}
-			}
-			return FileBlobs;
-		}
-
-		static void ExtractBlobs(Stream PackStream, Dictionary<string, DependencyBlob> Blobs)
-		{
-			List<BlobState> BlobStates = new List<BlobState>();
+			int MinFileIdx = 0;
+			int MaxFileIdx = 0;
+			FileStream[] OutputStreams = new FileStream[Files.Length];
+			SHA1Managed[] OutputHashers = new SHA1Managed[Files.Length];
 			try
 			{
-				// Prepare output files sorted by offset.
-				foreach (var Pair in Blobs)
-				{
-					BlobStates.Add(new BlobState(Pair.Key, Pair.Value));
-				}
-				BlobStates.Sort((a, b) => a.Blob.PackOffset.CompareTo(b.Blob.PackOffset));
-
 				// Create files from pack.
 				byte[] Buffer = new byte[16384];
 				long PackOffset = 0;
-				while (BlobStates.Count > 0)
+				while(MinFileIdx < Files.Length)
 				{
+					// Read the next chunk of data
 					int ReadSize;
 					try
 					{
-						ReadSize = PackStream.Read(Buffer, 0, Buffer.Length);
+						ReadSize = RawStream.Read(Buffer, 0, Buffer.Length);
 					}
-					catch (Exception e)
+					catch (Exception Ex)
 					{
-						throw new CorruptPackFileException("Can't read from pack stream", e);
+						throw new CorruptPackFileException("Can't read from pack stream", Ex);
 					}
 					if (ReadSize == 0)
 					{
 						throw new CorruptPackFileException("Unexpected end of file", null);
 					}
-					for (int i = 0; i < BlobStates.Count; ++i)
+
+					// Write to all the active files
+					for(int Idx = MinFileIdx; Idx < Files.Length && Files[Idx].MinPackOffset <= PackOffset + ReadSize; Idx++)
 					{
-						var State = BlobStates[i];
-						if (State.Blob.PackOffset > PackOffset + ReadSize) break;
-						if (State.Stream == null)
+						IncomingFile CurrentFile = Files[Idx];
+
+						// Open the stream if it's a new file
+						if(Idx == MaxFileIdx)
 						{
-							// Create the output folder
-							Directory.CreateDirectory(Path.GetDirectoryName(State.Path));
-							// Create file
-							State.Stream = File.OpenWrite(State.Path + IncomingFileSuffix);
-							State.Stream.SetLength(State.Blob.Size);
-							State.Hash.Initialize();
+							Directory.CreateDirectory(Path.GetDirectoryName(CurrentFile.Name));
+							OutputStreams[Idx] = File.Open(CurrentFile.Name + IncomingFileSuffix, FileMode.Create, FileAccess.Write, FileShare.None);
+							OutputHashers[Idx] = new SHA1Managed();
+							MaxFileIdx++;
 						}
-						int BufferOffset = (int)Math.Max(0, State.Blob.PackOffset - PackOffset);
-						int BufferCount = (int)Math.Min(ReadSize, State.Blob.PackOffset - PackOffset + State.Blob.Size) - BufferOffset;
-						State.Stream.Write(Buffer, BufferOffset, BufferCount);
-						State.Hash.TransformBlock(Buffer, BufferOffset, BufferCount, Buffer, BufferOffset);
-						if (State.Blob.PackOffset + State.Blob.Size <= PackOffset + ReadSize) 
+
+						// Write the data to this file
+						int BufferOffset = (int)Math.Max(0, CurrentFile.MinPackOffset - PackOffset);
+						int BufferCount = (int)Math.Min(ReadSize, CurrentFile.MaxPackOffset - PackOffset) - BufferOffset;
+						OutputStreams[Idx].Write(Buffer, BufferOffset, BufferCount);
+						OutputHashers[Idx].TransformBlock(Buffer, BufferOffset, BufferCount, Buffer, BufferOffset);
+
+						// If we're finished, verify the hash and close it
+						if(Idx == MinFileIdx && CurrentFile.MaxPackOffset <= PackOffset + ReadSize)
 						{
-							State.Hash.TransformFinalBlock(Buffer, 0, 0);
-							// Check the hash was what we expected
-							string ActualHash = BitConverter.ToString(State.Hash.Hash).ToLower().Replace("-", "");
-							if (ActualHash != State.Blob.Hash)
+							OutputHashers[Idx].TransformFinalBlock(Buffer, 0, 0);
+
+							string Hash = BitConverter.ToString(OutputHashers[Idx].Hash).ToLower().Replace("-", "");
+							if(Hash != CurrentFile.Hash)
 							{
-								throw new CorruptPackFileException("Incorrect hash value of file: " + State.Path, null);
+								throw new CorruptPackFileException(String.Format("Incorrect hash value of {0}: expected {1}, got {2}", CurrentFile.Name, CurrentFile.Hash, Hash), null);
 							}
-							State.Stream.Close();
-							State.Stream = null;
-							
-							// Move the file to its final position
-							File.Delete(State.Path);
-							File.Move(State.Path + IncomingFileSuffix, State.Path);
-
-							// Remove file from queue
-							BlobStates.RemoveAt(i);
-							--i;
+						
+							OutputStreams[Idx].Dispose();
+							File.Delete(CurrentFile.Name);
+							File.Move(CurrentFile.Name + IncomingFileSuffix, CurrentFile.Name);
+							MinFileIdx++;
 						}
 					}
 					PackOffset += ReadSize;
@@ -1038,13 +1028,10 @@ namespace GitDependencies
 			finally 
 			{
 				// Delete unfinished files.
-				foreach (var State in BlobStates)
+				for(int Idx = MinFileIdx; Idx < MaxFileIdx; Idx++)
 				{
-					if (State.Stream != null)
-					{
-						State.Stream.Close();
-						SafeDeleteFileQuiet(State.Path + IncomingFileSuffix);
-					}
+					OutputStreams[Idx].Dispose();
+					SafeDeleteFileQuiet(Files[Idx].Name + IncomingFileSuffix);
 				}
 			}
 		}
