Avoid race conditions with poisoned cache files.

diff --git a/Engine/Source/Programs/GitDependencies/Program.cs b/Engine/Source/Programs/GitDependencies/Program.cs
--- a/Engine/Source/Programs/GitDependencies/Program.cs
+++ b/Engine/Source/Programs/GitDependencies/Program.cs
@@ -839,24 +839,32 @@ namespace GitDependencies
 
 		static bool TryUnpackFromCache(string CacheFileName, long CompressedSize, IncomingFile[] Files)
 		{
-			if (CacheFileName != null)
+			if (CacheFileName != null && File.Exists(CacheFileName))
 			{
-				FileInfo PackFileInfo = new FileInfo(CacheFileName);
-				if (PackFileInfo.Exists && (PackFileInfo.Length == CompressedSize))
+				// Try to open the cached file for reading. Could fail due to race conditions despite checking above, so swallow any exceptions.
+				FileStream InputStream;
+				try
 				{
-					try
-					{
-						// Don't download file, if already exists.
-						using (GZipStream DecompressedStream = new GZipStream(PackFileInfo.OpenRead(), CompressionMode.Decompress, false))
-						{
-							ExtractFilesFromRawStream(DecompressedStream, Files);
-							return true;
-						}
-					}
-					catch (CorruptPackFileException)
-					{
-						SafeDeleteFileQuiet(CacheFileName);
-					}
+					InputStream = File.Open(CacheFileName, FileMode.Open, FileAccess.Read, FileShare.Read | FileShare.Delete);
+				}
+				catch(Exception)
+				{
+					return false;
+				}
+
+				// Try to extract files from the cache. If we get a corrupt pack file exception, delete it.
+				try
+				{
+					ExtractFiles(InputStream, Files);
+					return true;
+				}
+				catch(CorruptPackFileException)
+				{
+					SafeDeleteFileQuiet(CacheFileName);
+				}
+				finally
+				{
+					InputStream.Dispose();
 				}
 			}
 			return false;
