Change status message to match the format Git uses, for consistency.

diff --git a/Engine/Source/Programs/GitDependencies/Program.cs b/Engine/Source/Programs/GitDependencies/Program.cs
--- a/Engine/Source/Programs/GitDependencies/Program.cs
+++ b/Engine/Source/Programs/GitDependencies/Program.cs
@@ -728,10 +728,22 @@ namespace GitDependencies
 				long NumBytesTotal = Interlocked.Read(ref State.NumBytesTotal);
 				long NumBytesCached = Interlocked.Read(ref State.NumBytesCached);
 				long NumBytesPerSecond = (long)Math.Ceiling((float)Math.Max(NumBytesRead - NumBytesReadBuffer[BufferIdx], 0) * 1000.0f / (NumBytesReadBuffer.Length * TickInterval));
-				NumFilesReportedRead = State.NumFilesRead;
-				Log.WriteStatus("Received {0}/{1} files ({2}/{3}mb; {4}mb/s; {5}%; {6}mb from cache)...", NumFilesReportedRead, State.NumFiles, FormatMegabytes(NumBytesRead, 1), FormatMegabytes(NumBytesTotal, 1), FormatMegabytes(NumBytesPerSecond, 2), ((NumBytesRead + NumBytesCached) * 100) / (NumBytesTotal + NumBytesCached), FormatMegabytes(NumBytesCached, 1));
 
+				NumFilesReportedRead = State.NumFilesRead;
 				NumBytesReadBuffer[BufferIdx] = NumBytesRead;
+
+				StringBuilder Status = new StringBuilder();
+				Status.AppendFormat("Updating dependencies: {0,3}% ({1}/{2})", ((NumBytesRead + NumBytesCached) * 100) / (NumBytesTotal + NumBytesCached), NumFilesReportedRead, State.NumFiles);
+				if(NumBytesRead > 0)
+				{
+					Status.AppendFormat(", {0}/{1} MiB | {2}", FormatMegabytes(NumBytesRead, 1), FormatMegabytes(NumBytesTotal, 1), FormatMegabytes(NumBytesPerSecond, 2));
+				}
+				if(NumBytesCached > 0)
+				{
+					Status.AppendFormat(", {0} MiB cached", FormatMegabytes(NumBytesCached, 1));
+				}
+				Status.Append((NumFilesReportedRead == State.NumFiles)? ", done." : "...");
+				Log.WriteStatus(Status.ToString());
 			}
 
 			// If we finished with an error, try to clean up and return
