Remove unnecessary call to ReadToEnd(); ExtractFilesFromRawStream will now read the whole file if it's being hashed.

diff --git a/Engine/Source/Programs/GitDependencies/Program.cs b/Engine/Source/Programs/GitDependencies/Program.cs
--- a/Engine/Source/Programs/GitDependencies/Program.cs
+++ b/Engine/Source/Programs/GitDependencies/Program.cs
@@ -959,7 +959,6 @@ namespace GitDependencies
 					using(GZipStream DecompressedStream = new GZipStream(InputStream, CompressionMode.Decompress, true))
 					{
 						ExtractFilesFromRawStream(DecompressedStream, Files, Hasher);
-						ReadToEnd(DecompressedStream);
 					}
 				}
 
@@ -979,12 +978,6 @@ namespace GitDependencies
 			}
 		}
 
-		static void ReadToEnd(Stream InputStream)
-		{
-			byte[] Buffer = new byte[16384];
-			while (InputStream.Read(Buffer, 0, Buffer.Length) != 0) { }
-		}
-
 		static void ExtractFilesFromRawStream(Stream RawStream, IncomingFile[] Files, SHA1 RawStreamHasher)
 		{
 			int MinFileIdx = 0;
