PR #927: GitDependencies: reduce file fragmentation (usefull for HDD). (Contributed by [REDACTED_PERSON])

diff --git a/Engine/Source/Programs/GitDependencies/Program.cs b/Engine/Source/Programs/GitDependencies/Program.cs
--- a/Engine/Source/Programs/GitDependencies/Program.cs
+++ b/Engine/Source/Programs/GitDependencies/Program.cs
@@ -824,7 +824,7 @@ namespace GitDependencies
 					}
 					else
 					{
-						DownloadAndExtractFiles(NextPack.Url, NextPack.Proxy, NextPack.CacheFileName, NextPack.Hash, NextPack.Files, Size => { RollbackSize += Size; Interlocked.Add(ref State.NumBytesRead, Size); });
+						DownloadAndExtractFiles(NextPack.Url, NextPack.Proxy, NextPack.CacheFileName, NextPack.CompressedSize, NextPack.Hash, NextPack.Files, Size => { RollbackSize += Size; Interlocked.Add(ref State.NumBytesRead, Size); });
 					}
 
 					// Update the stats
@@ -891,7 +891,7 @@ namespace GitDependencies
 			return false;
 		}
 
-		static void DownloadAndExtractFiles(string Url, Uri Proxy, string CacheFileName, string ExpectedHash, IncomingFile[] Files, NotifyReadDelegate NotifyRead)
+		static void DownloadAndExtractFiles(string Url, Uri Proxy, string CacheFileName, long CompressedSize, string ExpectedHash, IncomingFile[] Files, NotifyReadDelegate NotifyRead)
 		{
 			// Create the web request
 			WebRequest Request = WebRequest.Create(Url);
@@ -915,7 +915,7 @@ namespace GitDependencies
 					}
 					else
 					{
-						ExtractFilesThroughCache(ResponseStream, CacheFileName, ExpectedHash, Files);
+						ExtractFilesThroughCache(ResponseStream, CacheFileName, CompressedSize, ExpectedHash, Files);
 					}
 				}
 			}
@@ -942,7 +942,7 @@ namespace GitDependencies
 			ExtractFilesFromRawStream(DecompressedStream, Files, null);
 		}
 
-		static void ExtractFilesThroughCache(Stream InputStream, string FileName, string ExpectedHash, IncomingFile[] Files)
+		static void ExtractFilesThroughCache(Stream InputStream, string FileName, long CompressedSize, string ExpectedHash, IncomingFile[] Files)
 		{
 			// Extract files from a pack file while writing to the cache file at the same time
 			string IncomingFileName = String.Format("{0}-{1}{2}", FileName, InstanceSuffix, IncomingFileSuffix);
@@ -955,6 +955,7 @@ namespace GitDependencies
 				SHA1 Hasher = SHA1.Create();
 				using(FileStream CacheStream = File.Open(IncomingFileName, FileMode.Create, FileAccess.Write, FileShare.None))
 				{
+					CacheStream.SetLength(CompressedSize);
 					ForkReadStream ForkedInputStream = new ForkReadStream(InputStream, CacheStream);
 					using(GZipStream DecompressedStream = new GZipStream(InputStream, CompressionMode.Decompress, true))
 					{
@@ -1022,6 +1023,7 @@ namespace GitDependencies
 						{
 							Directory.CreateDirectory(Path.GetDirectoryName(CurrentFile.Name));
 							OutputStreams[Idx] = File.Open(CurrentFile.Name + IncomingFileSuffix, FileMode.Create, FileAccess.Write, FileShare.None);
+							OutputStreams[Idx].SetLength(CurrentFile.MaxPackOffset - CurrentFile.MinPackOffset);
 							OutputHashers[Idx] = SHA1.Create();
 							MaxFileIdx++;
 						}
