Exclude Linux folders when syncing Git dependencies on Windows if LINUX_ROOT is not set.

diff --git a/Engine/Source/Programs/GitDependencies/Program.cs b/Engine/Source/Programs/GitDependencies/Program.cs
--- a/Engine/Source/Programs/GitDependencies/Program.cs
+++ b/Engine/Source/Programs/GitDependencies/Program.cs
@@ -181,6 +181,10 @@ namespace GitDependencies
 				{
 					ExcludeFolders.Add("Android");
 				}
+				if(Environment.OSVersion.Platform == PlatformID.Win32NT && Environment.GetEnvironmentVariable("LINUX_ROOT") == null)
+				{
+					ExcludeFolders.Add("Linux");
+				}
 			}
 
 			// Parse all the explicit include filters
