PR #943: GitDependencies: Adding metadata to a cleaner XML serialization (Contributed by [REDACTED_PERSON])

diff --git a/Engine/Source/Programs/GitDependencies/DependencyManifest.cs b/Engine/Source/Programs/GitDependencies/DependencyManifest.cs
--- a/Engine/Source/Programs/GitDependencies/DependencyManifest.cs
+++ b/Engine/Source/Programs/GitDependencies/DependencyManifest.cs
@@ -8,6 +8,7 @@ using System.Linq;
 using System.Text;
 using System.Threading.Tasks;
 using System.Xml.Serialization;
+using System.ComponentModel;
 
 namespace GitDependencies
 {
@@ -21,6 +22,7 @@ namespace GitDependencies
 		public string Hash;
 
 		[XmlAttribute]
+		[DefaultValue(false)]
 		public bool IsExecutable;
 	}
 
@@ -62,6 +64,7 @@ namespace GitDependencies
 		public string BaseUrl = "http://cdn.unrealengine.com/dependencies";
 
 		[XmlAttribute]
+		[DefaultValue(false)]
 		public bool IgnoreProxy;
 
 		[XmlArrayItem("File")]
