PR #928: Add GitDependencies cache size limit (LRU) (Contributed by [REDACTED_PERSON])

diff --git a/Engine/Source/Programs/GitDependencies/Program.cs b/Engine/Source/Programs/GitDependencies/Program.cs
--- a/Engine/Source/Programs/GitDependencies/Program.cs
+++ b/Engine/Source/Programs/GitDependencies/Program.cs
@@ -60,10 +60,16 @@ namespace GitDependencies
 			public DependencyManifest Manifest;
 			public DependencyPack Pack;
 
-			public DependencyPackInfo(DependencyManifest Manifest, DependencyPack Pack) {
+			public DependencyPackInfo(DependencyManifest Manifest, DependencyPack Pack) 
+			{
 				this.Manifest = Manifest;
 				this.Pack = Pack;
 			}
+
+			public string GetCacheFileName() 
+			{
+				return Path.Combine(Pack.Hash.Substring(0, 2), Pack.Hash);
+			}
 		}
 
 		class CorruptPackFileException : Exception
@@ -93,6 +99,8 @@ namespace GitDependencies
 			int MaxRetries = int.Parse(ParseParameter(ArgsList, "-max-retries=", "4"));
 			bool bDryRun = ParseSwitch(ArgsList, "-dry-run");
 			bool bHelp = ParseSwitch(ArgsList, "-help");
+			float CacheSizeMultiplier = float.Parse(ParseParameter(ArgsList, "-cache-size-multiplier=", "2"));
+			int CacheDays = int.Parse(ParseParameter(ArgsList, "-cache-days=", "3"));
 			string RootPath = ParseParameter(ArgsList, "-root=", Path.GetFullPath(Path.Combine(Path.GetDirectoryName(Assembly.GetExecutingAssembly().Location), "../../..")));
 
 			// Parse the cache path. A specific path can be set using -catch=<PATH> or the UE4_GITDEPS environment variable, otherwise we look for a parent .git directory
@@ -222,15 +230,17 @@ namespace GitDependencies
 				Log.WriteLine("   --root=<PATH>                 Set the repository directory to be sync");
 				Log.WriteLine("   --threads=<N>                 Use N threads when downloading new files");
 				Log.WriteLine("   --dry-run                     Print a list of outdated files and exit");
-				Log.WriteLine("   --max-retries                 Set the maximum number of retries for each file");
+				Log.WriteLine("   --max-retries                 Override maximum number of retries per file");
 				Log.WriteLine("   --proxy=<user:password@url>   Sets the HTTP proxy address and credentials");
 				Log.WriteLine("   --cache=<PATH>                Specifies a custom path for the download cache");
+				Log.WriteLine("   --cache-size-multiplier=<N>   Cache size as multiplier of current download");
+				Log.WriteLine("   --cache-days=<N>              Number of days to keep entries in the cache");
 				Log.WriteLine("   --no-cache                    Disable caching of downloaded files");
 				Log.WriteLine();
 				Log.WriteLine("Detected settings:");
 				Log.WriteLine("   Excluded folders: {0}", (ExcludeFolders.Count == 0)? "none" : String.Join(", ", ExcludeFolders));
 				Log.WriteLine("   Proxy server: {0}", (Proxy == null)? "none" : Proxy.ToString());
-				Log.WriteLine("   Current cache path: {0}", CachePath != null ? CachePath : "cache disabled");
+				Log.WriteLine("   Download cache: {0}", (CachePath == null)? "disabled" : CachePath);
 				return 0;
 			}
 
@@ -238,7 +248,7 @@ namespace GitDependencies
 			Console.CancelKeyPress += delegate { Log.FlushStatus(); };
 
 			// Update the tree. Make sure we clear out the status line if we quit for any reason (eg. ctrl-c)
-			if(!UpdateWorkingTree(bDryRun, RootPath, ExcludeFolders, NumThreads, MaxRetries, Proxy, Overwrite, CachePath))
+			if(!UpdateWorkingTree(bDryRun, RootPath, ExcludeFolders, NumThreads, MaxRetries, Proxy, Overwrite, CachePath, CacheSizeMultiplier, CacheDays))
 			{
 				return 1;
 			}
@@ -286,7 +296,7 @@ namespace GitDependencies
 			}
 		}
 
-		static bool UpdateWorkingTree(bool bDryRun, string RootPath, HashSet<string> ExcludeFolders, int NumThreads, int MaxRetries, Uri Proxy, OverwriteMode Overwrite, string CachePath)
+		static bool UpdateWorkingTree(bool bDryRun, string RootPath, HashSet<string> ExcludeFolders, int NumThreads, int MaxRetries, Uri Proxy, OverwriteMode Overwrite, string CachePath, float CacheSizeMultiplier, int CacheDays)
 		{
 			// Start scanning on the working directory 
 			if(ExcludeFolders.Count > 0)
@@ -555,6 +565,12 @@ namespace GitDependencies
 				{
 					return false;
 				}
+
+				// Cleanup cache files
+				if(CachePath != null)
+				{
+					PurgeCacheFiles(CachePath, TargetPacks, CacheSizeMultiplier, CacheDays);
+				}
 			}
 
 			// Update all the executable permissions
@@ -566,6 +582,42 @@ namespace GitDependencies
 			return true;
 		}
 
+		static void PurgeCacheFiles(string CachePath, Dictionary<string, DependencyPackInfo> Packs, float CacheSizeMultiplier, int CacheDays)
+		{
+			// Update the timestamp for all referenced packs
+			DateTime CurrentTime = DateTime.UtcNow;
+			foreach(DependencyPackInfo Pack in Packs.Values)
+			{
+				string FileName = Path.Combine(CachePath, Pack.GetCacheFileName());
+				if(File.Exists(FileName))
+				{
+					try { File.SetLastWriteTimeUtc(FileName, CurrentTime); } catch { }
+				}
+			}
+
+			// Get the size of the cache, and time before which we'll consider deleting entries
+			long DesiredCacheSize = (long)(Packs.Values.Sum(x => x.Pack.CompressedSize) * CacheSizeMultiplier);
+			DateTime StaleTime = CurrentTime - TimeSpan.FromDays(CacheDays) - TimeSpan.FromSeconds(5); // -5s for filesystems that don't store exact timestamps, like FAT.
+
+			// Enumerate all the files in the cache, and sort them by last write time
+			DirectoryInfo CacheDirectory = new DirectoryInfo(CachePath);
+			IEnumerable<FileInfo> CacheFiles = CacheDirectory.EnumerateFiles("*", SearchOption.AllDirectories);
+
+			// Find all the files in the cache directory
+			long CacheSize = 0;
+			foreach(FileInfo StaleFile in CacheFiles.OrderByDescending(x => x.LastWriteTimeUtc))
+			{
+				if(CacheSize > DesiredCacheSize && StaleFile.LastWriteTimeUtc < StaleTime)
+				{
+					StaleFile.Delete();
+				}
+				else
+				{
+					CacheSize += StaleFile.Length;
+				}
+			}
+		}
+
 		static bool SetExecutablePermissions(string RootDir, IEnumerable<DependencyFile> Files)
 		{
 			// Try to load the Mono Posix assembly. If it doesn't exist, we're on Windows.
@@ -706,7 +758,7 @@ namespace GitDependencies
 				Pack.Url = String.Format("{0}/{1}/{2}", RequiredPack.Manifest.BaseUrl, RequiredPack.Pack.RemotePath, RequiredPack.Pack.Hash);
 				Pack.Proxy = RequiredPack.Manifest.IgnoreProxy? null : Proxy;
 				Pack.Hash = RequiredPack.Pack.Hash;
-				Pack.CacheFileName = (CachePath == null)? null : Path.Combine(CachePath, RequiredPack.Pack.Hash.Substring(0, 2), RequiredPack.Pack.Hash);
+				Pack.CacheFileName = (CachePath == null)? null : Path.Combine(CachePath, RequiredPack.GetCacheFileName());
 				Pack.Files = GetIncomingFilesForPack(RootPath, RequiredPack.Pack, PackToBlobs, BlobToFiles);
 				Pack.CompressedSize = RequiredPack.Pack.CompressedSize;
 				DownloadQueue.Enqueue(Pack);
