Increase cache minimum time to 7 days.

diff --git a/Engine/Source/Programs/GitDependencies/Program.cs b/Engine/Source/Programs/GitDependencies/Program.cs
--- a/Engine/Source/Programs/GitDependencies/Program.cs
+++ b/Engine/Source/Programs/GitDependencies/Program.cs
@@ -100,7 +100,7 @@ namespace GitDependencies
 			bool bDryRun = ParseSwitch(ArgsList, "-dry-run");
 			bool bHelp = ParseSwitch(ArgsList, "-help");
 			float CacheSizeMultiplier = float.Parse(ParseParameter(ArgsList, "-cache-size-multiplier=", "2"));
-			int CacheDays = int.Parse(ParseParameter(ArgsList, "-cache-days=", "3"));
+			int CacheDays = int.Parse(ParseParameter(ArgsList, "-cache-days=", "7"));
 			string RootPath = ParseParameter(ArgsList, "-root=", Path.GetFullPath(Path.Combine(Path.GetDirectoryName(Assembly.GetExecutingAssembly().Location), "../../..")));
 
 			// Parse the cache path. A specific path can be set using -catch=<PATH> or the UE4_GITDEPS environment variable, otherwise we look for a parent .git directory
@@ -597,7 +597,7 @@ namespace GitDependencies
 
 			// Get the size of the cache, and time before which we'll consider deleting entries
 			long DesiredCacheSize = (long)(Packs.Values.Sum(x => x.Pack.CompressedSize) * CacheSizeMultiplier);
-			DateTime StaleTime = CurrentTime - TimeSpan.FromDays(CacheDays) - TimeSpan.FromSeconds(5); // -5s for filesystems that don't store exact timestamps, like FAT.
+			DateTime StaleTime = CurrentTime - TimeSpan.FromDays(CacheDays) - TimeSpan.FromSeconds(5); // +5s for filesystems that don't store exact timestamps, like FAT.
 
 			// Enumerate all the files in the cache, and sort them by last write time
 			DirectoryInfo CacheDirectory = new DirectoryInfo(CachePath);
