Add validation for integer and float arguments to GitDependencies.

diff --git a/Engine/Source/Programs/GitDependencies/Program.cs b/Engine/Source/Programs/GitDependencies/Program.cs
--- a/Engine/Source/Programs/GitDependencies/Program.cs
+++ b/Engine/Source/Programs/GitDependencies/Program.cs
@@ -96,12 +96,12 @@ namespace GitDependencies
 			NormalizeArguments(DefaultArgsList);
 
 			// Parse the parameters
-			int NumThreads = int.Parse(ParseParameter(ArgsList, DefaultArgsList, "-threads=", "4"));
-			int MaxRetries = int.Parse(ParseParameter(ArgsList, DefaultArgsList, "-max-retries=", "4"));
+			int NumThreads = ParseIntParameter(ArgsList, DefaultArgsList, "-threads=", 4);
+			int MaxRetries = ParseIntParameter(ArgsList, DefaultArgsList, "-max-retries=", 4);
 			bool bDryRun = ParseSwitch(ArgsList, "-dry-run");
 			bool bHelp = ParseSwitch(ArgsList, "-help");
-			float CacheSizeMultiplier = float.Parse(ParseParameter(ArgsList, DefaultArgsList, "-cache-size-multiplier=", "2"));
-			int CacheDays = int.Parse(ParseParameter(ArgsList, DefaultArgsList, "-cache-days=", "7"));
+			float CacheSizeMultiplier = ParseFloatParameter(ArgsList, DefaultArgsList, "-cache-size-multiplier=", 2.0f);
+			int CacheDays = ParseIntParameter(ArgsList, DefaultArgsList, "-cache-days=", 7);
 			string RootPath = ParseParameter(ArgsList, "-root=", Path.GetFullPath(Path.Combine(Path.GetDirectoryName(Assembly.GetExecutingAssembly().Location), "../../..")));
 
 			// Parse the cache path. A specific path can be set using -catch=<PATH> or the UE4_GITDEPS environment variable, otherwise we look for a parent .git directory
@@ -334,6 +334,44 @@ namespace GitDependencies
 			return ParseParameter(ArgsList, Prefix, ParseParameter(DefaultArgsList, Prefix, Default));
 		}
 
+		static int ParseIntParameter(List<string> ArgsList, string Prefix, int Default)
+		{
+			for(int Idx = 0; Idx < ArgsList.Count; Idx++)
+			{
+				int Value;
+				if(ArgsList[Idx].StartsWith(Prefix, StringComparison.CurrentCultureIgnoreCase) && int.TryParse(ArgsList[Idx].Substring(Prefix.Length), out Value))
+				{
+					ArgsList.RemoveAt(Idx);
+					return Value;
+				}
+			}
+			return Default;
+		}
+
+		static int ParseIntParameter(List<string> ArgsList, List<string> DefaultArgsList, string Prefix, int Default)
+		{
+			return ParseIntParameter(ArgsList, Prefix, ParseIntParameter(DefaultArgsList, Prefix, Default));
+		}
+
+		static float ParseFloatParameter(List<string> ArgsList, string Prefix, float Default)
+		{
+			for(int Idx = 0; Idx < ArgsList.Count; Idx++)
+			{
+				float Value;
+				if(ArgsList[Idx].StartsWith(Prefix, StringComparison.CurrentCultureIgnoreCase) && float.TryParse(ArgsList[Idx].Substring(Prefix.Length), out Value))
+				{
+					ArgsList.RemoveAt(Idx);
+					return Value;
+				}
+			}
+			return Default;
+		}
+
+		static float ParseFloatParameter(List<string> ArgsList, List<string> DefaultArgsList, string Prefix, float Default)
+		{
+			return ParseFloatParameter(ArgsList, Prefix, ParseFloatParameter(DefaultArgsList, Prefix, Default));
+		}
+
 		static IEnumerable<string> ParseParameters(List<string> ArgsList, string Prefix)
 		{
 			for(;;)
