PR #1073: Hacky fix so GitDependencies can exclude files in the root directory. (Contributed by [REDACTED_PERSON])

diff --git a/Engine/Source/Programs/GitDependencies/Program.cs b/Engine/Source/Programs/GitDependencies/Program.cs
--- a/Engine/Source/Programs/GitDependencies/Program.cs
+++ b/Engine/Source/Programs/GitDependencies/Program.cs
@@ -792,9 +792,10 @@ namespace GitDependencies
 
 		static bool IsExcludedFolder(string Name, IEnumerable<string> ExcludeFolders)
 		{
+			string RootedName = "/" + Name;
 			foreach(string ExcludeFolder in ExcludeFolders)
 			{
-				if(Name.IndexOf("/" + ExcludeFolder + "/", StringComparison.CurrentCultureIgnoreCase) != -1)
+				if(RootedName.IndexOf("/" + ExcludeFolder + "/", StringComparison.CurrentCultureIgnoreCase) != -1)
 				{
 					return true;
 				}
