PR #1131: Support prefix matching for gitdependencies' --exclude parameter. (Contributed by [REDACTED_PERSON])

diff --git a/Engine/Source/Programs/GitDependencies/Program.cs b/Engine/Source/Programs/GitDependencies/Program.cs
--- a/Engine/Source/Programs/GitDependencies/Program.cs
+++ b/Engine/Source/Programs/GitDependencies/Program.cs
@@ -199,13 +199,13 @@ namespace GitDependencies
 			// Parse all the explicit include filters
 			foreach(string IncludeFolder in ParseParameters(ArgsList, "-include="))
 			{
-				ExcludeFolders.Remove(IncludeFolder.Replace('\\', '/').Trim('/'));
+				ExcludeFolders.Remove(IncludeFolder.Replace('\\', '/').TrimEnd('/'));
 			}
 
 			// Parse all the explicit exclude filters
 			foreach(string ExcludeFolder in ParseParameters(ArgsList, "-exclude="))
 			{
-				ExcludeFolders.Add(ExcludeFolder.Replace('\\', '/').Trim('/'));
+				ExcludeFolders.Add(ExcludeFolder.Replace('\\', '/').TrimEnd('/'));
 			}
 
 			// If there are any more parameters, print an error
@@ -795,9 +795,21 @@ namespace GitDependencies
 			string RootedName = "/" + Name;
 			foreach(string ExcludeFolder in ExcludeFolders)
 			{
-				if(RootedName.IndexOf("/" + ExcludeFolder + "/", StringComparison.CurrentCultureIgnoreCase) != -1)
+				if(ExcludeFolder.StartsWith("/"))
 				{
-					return true;
+					// Do a prefix check
+					if(RootedName.StartsWith(ExcludeFolder + "/", StringComparison.CurrentCultureIgnoreCase))
+					{
+						return true;
+					}
+				}
+				else
+				{
+					// Do a substring check
+					if(RootedName.IndexOf("/" + ExcludeFolder + "/", StringComparison.CurrentCultureIgnoreCase) != -1)
+					{
+						return true;
+					}
 				}
 			}
 			return false;
