Fix publish script for GitDependencies.

MSBuild does not seem to track the renamed apphost executable when creating self-contained apps for different platforms, resulting in incorrect platform being published. Manually clear the bin and obj directories to ensure it always copies the correct file.

diff --git a/Engine/Source/Programs/GitDependencies/Publish.bat b/Engine/Source/Programs/GitDependencies/Publish.bat
--- a/Engine/Source/Programs/GitDependencies/Publish.bat
+++ b/Engine/Source/Programs/GitDependencies/Publish.bat
@@ -1,18 +1,28 @@
 @echo off
+setlocal
+pushd "%~dp0"
 call ..\..\..\Build\BatchFiles\GetDotnetPath.bat
 del /s /q "..\..\..\Binaries\DotNET\GitDependencies\*"
 
 echo.
 echo Building for win-x64...
-"%DOTNET_ROOT%\dotnet" publish GitDependencies.csproj -r win-x64 --output "..\..\..\Binaries\DotNET\GitDependencies\win-x64" --nologo
+rmdir /s /q bin
+rmdir /s /q obj
+"%DOTNET_ROOT%\dotnet" publish GitDependencies.csproj -r win-x64 -c Release --output "..\..\..\Binaries\DotNET\GitDependencies\win-x64" --nologo
 if errorlevel 1 goto :eof
 
 echo.
 echo Building for osx-x64...
-"%DOTNET_ROOT%\dotnet" publish GitDependencies.csproj -r osx-x64 --output "..\..\..\Binaries\DotNET\GitDependencies\osx-x64" --nologo
+rmdir /s /q bin
+rmdir /s /q obj
+"%DOTNET_ROOT%\dotnet" publish GitDependencies.csproj -r osx-x64 -c Release --output "..\..\..\Binaries\DotNET\GitDependencies\osx-x64" --nologo
 if errorlevel 1 goto :eof
 
 echo.
 echo Building for linux-x64...
-"%DOTNET_ROOT%\dotnet" publish GitDependencies.csproj -r linux-x64 --output "..\..\..\Binaries\DotNET\GitDependencies\linux-x64" --nologo
+rmdir /s /q bin
+rmdir /s /q obj
+"%DOTNET_ROOT%\dotnet" publish GitDependencies.csproj -r linux-x64 -c Release --output "..\..\..\Binaries\DotNET\GitDependencies\linux-x64" --nologo
 if errorlevel 1 goto :eof
+
+endlocal
