Make GitDependencies report any read-only files and prompt to overwrite like we do with modified files.

#rnx

diff --git a/Engine/Source/Programs/GitDependencies/Program.cs b/Engine/Source/Programs/GitDependencies/Program.cs
--- a/Engine/Source/Programs/GitDependencies/Program.cs
+++ b/Engine/Source/Programs/GitDependencies/Program.cs
@@ -103,7 +103,7 @@ namespace GitDependencies
 			bool bHelp = ParseSwitch(ArgsList, "-help");
 			float CacheSizeMultiplier = ParseFloatParameter(ArgsList, DefaultArgsList, "-cache-size-multiplier=", 2.0f);
 			int CacheDays = ParseIntParameter(ArgsList, DefaultArgsList, "-cache-days=", 7);
-			string RootPath = ParseParameter(ArgsList, "-root=", Path.GetFullPath(Path.Combine(Path.GetDirectoryName(Assembly.GetExecutingAssembly().Location), "../../../../..")));
+			string RootPath = ParseParameter(ArgsList, "-root=", Path.GetFullPath(Path.Combine(Path.GetDirectoryName(Assembly.GetExecutingAssembly().Location), "../../..")));
 
 			// Parse the cache path. A specific path can be set using -catch=<PATH> or the UE4_GITDEPS environment variable, otherwise we look for a parent .git directory
 			// and use a sub-folder of that. Users which download the source through a zip file (and won't have a .git directory) are unlikely to benefit from caching, as
@@ -233,10 +233,10 @@ namespace GitDependencies
 
 		static string GetLegacyEnvironmentVariable(string Name, string LegacyName)
 		{
-			string Value = System.Environment.GetEnvironmentVariable(Name);
+			string Value = Environment.GetEnvironmentVariable(Name);
 			if (string.IsNullOrEmpty(Value))
 			{
-				Value = System.Environment.GetEnvironmentVariable(LegacyName);
+				Value = Environment.GetEnvironmentVariable(LegacyName);
 			}
 			return Value;
 		}
@@ -452,36 +452,49 @@ namespace GitDependencies
 			}
 
 			// Find all the existing files in the working directory from previous runs. Use the working manifest to cache hashes for them based on timestamp, but recalculate them as needed.
+			List<WorkingFile> ReadOnlyFiles = new List<WorkingFile>();
 			Dictionary<string, WorkingFile> CurrentFileLookup = new Dictionary<string, WorkingFile>();
 			foreach(WorkingFile CurrentFile in CurrentManifest.Files)
 			{
 				// Update the hash for this file
 				string CurrentFilePath = Path.Combine(RootPath, CurrentFile.Name);
-				if(File.Exists(CurrentFilePath))
+				FileInfo CurrentFileInfo = new FileInfo(CurrentFilePath);
+				if(CurrentFileInfo.Exists)
 				{
-					long LastWriteTime = File.GetLastWriteTimeUtc(CurrentFilePath).Ticks;
+					long LastWriteTime = CurrentFileInfo.LastWriteTimeUtc.Ticks;
 					if(LastWriteTime != CurrentFile.Timestamp)
 					{
 						CurrentFile.Hash = ComputeHashForFile(CurrentFilePath);
 						CurrentFile.Timestamp = LastWriteTime;
 					}
 					CurrentFileLookup.Add(CurrentFile.Name, CurrentFile);
+
+					if (CurrentFileInfo.IsReadOnly)
+					{
+						ReadOnlyFiles.Add(CurrentFile);
+					}
 				}
 			}
 
 			// Also add all the untracked files which already exist, but weren't downloaded by this program
 			foreach (DependencyFile TargetFile in TargetFiles.Values) 
 			{
-				if(!CurrentFileLookup.ContainsKey(TargetFile.Name))
+				if (!CurrentFileLookup.ContainsKey(TargetFile.Name))
 				{
 					string CurrentFilePath = Path.Combine(RootPath, TargetFile.Name);
-					if(File.Exists(CurrentFilePath))
+					FileInfo CurrentFileInfo = new FileInfo(CurrentFilePath);
+					if (CurrentFileInfo.Exists)
 					{
 						WorkingFile CurrentFile = new WorkingFile();
 						CurrentFile.Name = TargetFile.Name;
 						CurrentFile.Hash = ComputeHashForFile(CurrentFilePath);
-						CurrentFile.Timestamp = File.GetLastWriteTimeUtc(CurrentFilePath).Ticks;
+						CurrentFile.Timestamp = CurrentFileInfo.LastWriteTimeUtc.Ticks;
 						CurrentFileLookup.Add(CurrentFile.Name, CurrentFile);
+
+						if (CurrentFileInfo.IsReadOnly)
+						{
+							ReadOnlyFiles.Add(CurrentFile);
+						}
 					}
 				}
 			}
@@ -501,7 +514,7 @@ namespace GitDependencies
 
 			// Create a new working manifest for the working directory, moving over files that we already have. Add any missing dependencies into the download queue.
 			WorkingManifest NewWorkingManifest = new WorkingManifest();
-			foreach(DependencyFile TargetFile in FilteredTargetFiles)
+			foreach (DependencyFile TargetFile in FilteredTargetFiles)
 			{
 				WorkingFile NewFile;
 				if(CurrentFileLookup.TryGetValue(TargetFile.Name, out NewFile) && NewFile.Hash == TargetFile.Hash)
@@ -511,6 +524,7 @@ namespace GitDependencies
 
 					// Move the existing file to the new working set
 					CurrentFileLookup.Remove(NewFile.Name);
+					ReadOnlyFiles.Remove(NewFile);
 				}
 				else
 				{
@@ -561,34 +575,53 @@ namespace GitDependencies
 				}
 			}
 
-			// Warn if there were any files that have been tampered with, and allow the user to choose whether to overwrite them
+			// Warn if there were any files that have been tampered with or are read only, and allow the user to choose whether to overwrite them
 			bool bOverwriteTamperedFiles = true;
-			if(TamperedFiles.Count > 0 && Overwrite != OverwriteMode.Force)
+			if (Overwrite != OverwriteMode.Force)
 			{
-				// List the files that have changed
-				Log.WriteError("The following file(s) have been modified:");
-				foreach(WorkingFile TamperedFile in TamperedFiles)
+				bool PromptForOverwrite = false;
+				if (TamperedFiles.Any())
 				{
-					Log.WriteError("  {0}", TamperedFile.Name);
+					PromptForOverwrite = true;
+					// List the files that have changed
+					Log.WriteError("The following file(s) have been modified:");
+					foreach (WorkingFile TamperedFile in TamperedFiles)
+					{
+						Log.WriteError("  {0}", TamperedFile.Name);
+					}
 				}
 
-				// Figure out whether to overwrite the files
-				if(Overwrite == OverwriteMode.Unchanged)
+				if (ReadOnlyFiles.Any())
 				{
-					Log.WriteError("Re-run with the --force parameter to overwrite them.");
-					bOverwriteTamperedFiles = false;
+					PromptForOverwrite = true;
+					// List the files that are read only
+					Log.WriteError("The following file(s) are flagged as read only:");
+					foreach (WorkingFile ReadOnlyFile in ReadOnlyFiles)
+					{
+						Log.WriteError("  {0}", ReadOnlyFile.Name);
+					}
 				}
-				else
+
+				// Figure out whether to overwrite the files
+				if (PromptForOverwrite)
 				{
-					Log.WriteStatus("Would you like to overwrite your changes (y/n)? ");
-					ConsoleKeyInfo KeyInfo = Console.ReadKey(false);
-					bOverwriteTamperedFiles = (KeyInfo.KeyChar == 'y' || KeyInfo.KeyChar == 'Y');
-					Log.FlushStatus();
+					if (Overwrite == OverwriteMode.Unchanged)
+					{
+						Log.WriteError("Re-run with the --force parameter to overwrite them.");
+						bOverwriteTamperedFiles = false;
+					}
+					else
+					{
+						Log.WriteStatus("Would you like to overwrite your changes (y/n)? ");
+						ConsoleKeyInfo KeyInfo = Console.ReadKey(false);
+						bOverwriteTamperedFiles = (KeyInfo.KeyChar == 'y' || KeyInfo.KeyChar == 'Y');
+						Log.FlushStatus();
+					}
 				}
 			}
 
 			// Overwrite any tampered files, or remove them from the download list
-			if(bOverwriteTamperedFiles)
+			if (bOverwriteTamperedFiles)
 			{
 				foreach(WorkingFile TamperedFile in TamperedFiles)
 				{
@@ -597,15 +630,25 @@ namespace GitDependencies
 						return false;
 					}
 				}
+
+				foreach (WorkingFile ReadOnlyFile in ReadOnlyFiles)
+				{
+					string FilePath = Path.Combine(RootPath, ReadOnlyFile.Name);
+					File.SetAttributes(FilePath, File.GetAttributes(FilePath) & ~FileAttributes.ReadOnly);
+					if (!SafeDeleteFile(FilePath))
+					{
+						return false;
+					}
+				}
 			}
 			else
 			{
-				foreach(WorkingFile TamperedFile in TamperedFiles)
+				foreach(WorkingFile FileToIgnore in TamperedFiles.Concat(ReadOnlyFiles))
 				{
 					DependencyFile TargetFile;
-					if(TargetFiles.TryGetValue(TamperedFile.Name, out TargetFile))
+					if(TargetFiles.TryGetValue(FileToIgnore.Name, out TargetFile))
 					{
-						TargetFiles.Remove(TamperedFile.Name);
+						TargetFiles.Remove(FileToIgnore.Name);
 						FilesToDownload.Remove(TargetFile);
 					}
 				}
