GitDependencies: enabling nullable

#rnx

diff --git a/Engine/Source/Programs/GitDependencies/DependencyManifest.cs b/Engine/Source/Programs/GitDependencies/DependencyManifest.cs
--- a/Engine/Source/Programs/GitDependencies/DependencyManifest.cs
+++ b/Engine/Source/Programs/GitDependencies/DependencyManifest.cs
@@ -1,4 +1,4 @@
-﻿// Copyright Epic Games, Inc. All Rights Reserved.
+// Copyright Epic Games, Inc. All Rights Reserved.
 
 using System;
 using System.Collections.Generic;
@@ -22,13 +22,13 @@ namespace GitDependencies
 		/// Name of the file, relative to the root. Directories are separated with forward slashes, without a leading slash.
 		/// </summary>
 		[XmlAttribute]
-		public string Name;
+		public string Name = null!;
 
 		/// <summary>
 		/// Hash of this file. Used to determine which blob represents its contents.
 		/// </summary>
 		[XmlAttribute]
-		public string Hash;
+		public string Hash = null!;
 
 		/// <summary>
 		/// Whether this file should have the exectuable bit set, on platforms that support it.
@@ -48,7 +48,7 @@ namespace GitDependencies
 		/// Hash of the blob.
 		/// </summary>
 		[XmlAttribute]
-		public string Hash;
+		public string Hash = null!;
 
 		/// <summary>
 		/// Size of the blob.
@@ -60,7 +60,7 @@ namespace GitDependencies
 		/// Hash of the pack file containing this blob.
 		/// </summary>
 		[XmlAttribute]
-		public string PackHash;
+		public string PackHash = null!;
 
 		/// <summary>
 		/// Offset of this blob's data within the pack file.
@@ -79,7 +79,7 @@ namespace GitDependencies
 		/// Hash of this pack file.
 		/// </summary>
 		[XmlAttribute]
-		public string Hash;
+		public string Hash = null!;
 
 		/// <summary>
 		/// Size of this pack file, when uncompressed.
@@ -97,7 +97,7 @@ namespace GitDependencies
 		/// Subdirectory for this pack file on the remote server.
 		/// </summary>
 		[XmlAttribute]
-		public string RemotePath;
+		public string RemotePath = null!;
 	}
 
 	/// <summary>
diff --git a/Engine/Source/Programs/GitDependencies/ForkReadStream.cs b/Engine/Source/Programs/GitDependencies/ForkReadStream.cs
--- a/Engine/Source/Programs/GitDependencies/ForkReadStream.cs
+++ b/Engine/Source/Programs/GitDependencies/ForkReadStream.cs
@@ -1,4 +1,4 @@
-﻿// Copyright Epic Games, Inc. All Rights Reserved.
+// Copyright Epic Games, Inc. All Rights Reserved.
 
 using System;
 using System.Collections.Generic;
@@ -10,10 +10,10 @@ namespace GitDependencies
 {
 	class ForkReadStream : Stream
 	{
-		Stream Inner;
-		Stream ForkWrite;
+		Stream? Inner;
+		Stream? ForkWrite;
 
-		public ForkReadStream(Stream InInner, Stream InForkWrite)
+		public ForkReadStream(Stream? InInner, Stream? InForkWrite)
 		{
 			Inner = InInner;
 			ForkWrite = InForkWrite;
@@ -52,7 +52,7 @@ namespace GitDependencies
 		{
 			get
 			{
-				return Inner.Position;
+				return Inner!.Position;
 			}
 			set
 			{
@@ -62,7 +62,7 @@ namespace GitDependencies
 
 		public override long Length
 		{
-			get { return Inner.Length; }
+			get { return Inner!.Length; }
 		}
 
 		public override void SetLength(long Value)
@@ -72,10 +72,10 @@ namespace GitDependencies
 
 		public override int Read(byte[] Buffer, int Offset, int Count)
 		{
-			int SizeRead = Inner.Read(Buffer, Offset, Count);
+			int SizeRead = Inner!.Read(Buffer, Offset, Count);
 			if (SizeRead > 0)
 			{
-				ForkWrite.Write(Buffer, Offset, SizeRead);
+				ForkWrite!.Write(Buffer, Offset, SizeRead);
 			}
 			return SizeRead;
 		}
@@ -92,7 +92,7 @@ namespace GitDependencies
 
 		public override void Flush()
 		{
-			ForkWrite.Flush();
+			ForkWrite!.Flush();
 		}
 	}
 }
diff --git a/Engine/Source/Programs/GitDependencies/GitDependencies.csproj b/Engine/Source/Programs/GitDependencies/GitDependencies.csproj
--- a/Engine/Source/Programs/GitDependencies/GitDependencies.csproj
+++ b/Engine/Source/Programs/GitDependencies/GitDependencies.csproj
@@ -10,8 +10,8 @@
     <Configurations>Debug;Release;Development</Configurations>
     <AppendRuntimeIdentifierToOutputPath>false</AppendRuntimeIdentifierToOutputPath>
     <DebugType>pdbonly</DebugType> 
-    <Nullable>disable</Nullable> 
-    <PlatformTarget>x64</PlatformTarget> 
+    <Nullable>enable</Nullable> 
+    <PlatformTarget>AnyCPU</PlatformTarget> 
     <StartupObject>GitDependencies.Program</StartupObject>
   </PropertyGroup>
   <PropertyGroup Condition=" '$(Configuration)|$(Platform)' == 'Development|AnyCPU' ">
diff --git a/Engine/Source/Programs/GitDependencies/NotifyReadStream.cs b/Engine/Source/Programs/GitDependencies/NotifyReadStream.cs
--- a/Engine/Source/Programs/GitDependencies/NotifyReadStream.cs
+++ b/Engine/Source/Programs/GitDependencies/NotifyReadStream.cs
@@ -1,4 +1,4 @@
-﻿// Copyright Epic Games, Inc. All Rights Reserved.
+// Copyright Epic Games, Inc. All Rights Reserved.
 
 using System;
 using System.Collections.Generic;
@@ -12,11 +12,11 @@ namespace GitDependencies
 
 	class NotifyReadStream : Stream
 	{
-		Stream Inner;
+		Stream? Inner;
 		NotifyReadDelegate NotifyRead;
 		long TotalRead;
 
-		public NotifyReadStream(Stream InInner, NotifyReadDelegate InNotifyRead)
+		public NotifyReadStream(Stream? InInner, NotifyReadDelegate InNotifyRead)
 		{
 			Inner = InInner;
 			NotifyRead = InNotifyRead;
@@ -61,7 +61,7 @@ namespace GitDependencies
 
 		public override long Length
 		{
-			get { return Inner.Length; }
+			get { return Inner!.Length; }
 		}
 
 		public override void SetLength(long Value)
@@ -71,7 +71,7 @@ namespace GitDependencies
 
 		public override int Read(byte[] Buffer, int Offset, int Count)
 		{
-			int SizeRead = Inner.Read(Buffer, Offset, Count);
+			int SizeRead = Inner!.Read(Buffer, Offset, Count);
 			NotifyRead(SizeRead);
 			TotalRead += SizeRead;
 			return SizeRead;
@@ -89,7 +89,7 @@ namespace GitDependencies
 
 		public override void Flush()
 		{
-			Inner.Flush();
+			Inner!.Flush();
 		}
 	}
 }
diff --git a/Engine/Source/Programs/GitDependencies/Program.cs b/Engine/Source/Programs/GitDependencies/Program.cs
--- a/Engine/Source/Programs/GitDependencies/Program.cs
+++ b/Engine/Source/Programs/GitDependencies/Program.cs
@@ -29,7 +29,7 @@ namespace GitDependencies
 			public long NumBytesTotal;
 			public long NumBytesCached;
 			public int NumFailingOrIdleDownloads;
-			public string LastDownloadError;
+			public string LastDownloadError = null!;
 		}
 
 		enum OverwriteMode
@@ -41,17 +41,17 @@ namespace GitDependencies
 
 		class IncomingPack
 		{
-			public string Url;
-			public string Hash;
-			public string CacheFileName;
-			public IncomingFile[] Files;
+			public string Url = null!;
+			public string Hash = null!;
+			public string? CacheFileName = null!;
+			public IncomingFile[] Files = new IncomingFile[0];
 			public long CompressedSize;
 		}
 
 		class IncomingFile
 		{
-			public string[] Names;
-			public string Hash;
+			public string[] Names = new string[0];
+			public string Hash = null!;
 			public long MinPackOffset;
 			public long MaxPackOffset;
 		}
@@ -75,7 +75,7 @@ namespace GitDependencies
 
 		class CorruptPackFileException : Exception
 		{
-			public CorruptPackFileException(string Message, Exception InnerException)
+			public CorruptPackFileException(string Message, Exception? InnerException)
 				: base(Message, InnerException)
 			{
 			}
@@ -88,7 +88,7 @@ namespace GitDependencies
 		const string LegacyManifestFilename = ".ue4dependencies";
 
 		static readonly string InstanceSuffix = Guid.NewGuid().ToString().Replace("-", "");
-		static HttpClient HttpClientInstance = null;
+		static HttpClient? HttpClientInstance = null;
 
 		static int Main(string[] Args)
 		{
@@ -107,19 +107,19 @@ namespace GitDependencies
 			bool bHelp = ParseSwitch(ArgsList, "-help");
 			float CacheSizeMultiplier = ParseFloatParameter(ArgsList, DefaultArgsList, "-cache-size-multiplier=", 2.0f);
 			int CacheDays = ParseIntParameter(ArgsList, DefaultArgsList, "-cache-days=", 7);
-			string RootPath = ParseParameter(ArgsList, "-root=", Path.GetFullPath(Path.Combine(Path.GetDirectoryName(Assembly.GetExecutingAssembly().Location), "../../../../..")));
+			string RootPath = ParseParameter(ArgsList, "-root=", Path.GetFullPath(Path.Combine(Path.GetDirectoryName(Assembly.GetExecutingAssembly()!.Location!)!, "../../../../..")))!;
 			double HttpTimeoutMultiplier = ParseFloatParameter(ArgsList, DefaultArgsList, "-http-timeout-multiplier=", 1.0f) * NumThreads;
 
 			// Parse the cache path. A specific path can be set using -catch=<PATH> or the UE4_GITDEPS environment variable, otherwise we look for a parent .git directory
 			// and use a sub-folder of that. Users which download the source through a zip file (and won't have a .git directory) are unlikely to benefit from caching, as
 			// they only need to download dependencies once.
-			string CachePath = null;
+			string? CachePath = null;
 			if (!ParseSwitch(ArgsList, "-no-cache"))
 			{
-				string CachePathParam = ParseParameter(ArgsList, DefaultArgsList, "-cache=", Environment.GetEnvironmentVariable("UE_GITDEPS"));
+				string? CachePathParam = ParseParameter(ArgsList, DefaultArgsList, "-cache=", Environment.GetEnvironmentVariable("UE_GITDEPS"));
 				if (String.IsNullOrEmpty(CachePathParam))
 				{
-					string CheckPath = Path.GetFullPath(RootPath);
+					string? CheckPath = Path.GetFullPath(RootPath);
 					while (CheckPath != null)
 					{
 						string GitPath = Path.Combine(CheckPath, ".git");
@@ -158,7 +158,7 @@ namespace GitDependencies
 			}
 
 			// Setup network proxy from argument list or environment variable
-			string ProxyUrl = ParseParameter(ArgsList, DefaultArgsList, "-proxy=", null);
+			string? ProxyUrl = ParseParameter(ArgsList, DefaultArgsList, "-proxy=", null);
 			if(String.IsNullOrEmpty(ProxyUrl))
 			{
 				ProxyUrl = Environment.GetEnvironmentVariable("HTTP_PROXY");
@@ -169,7 +169,7 @@ namespace GitDependencies
 			}
 
 			// Create a URI for the proxy. If there's no included username/password, accept them as separate parameters for legacy reasons.
-			Uri Proxy = null;
+			Uri? Proxy = null;
 			if(!String.IsNullOrEmpty(ProxyUrl))
 			{
 				UriBuilder ProxyBuilder = new UriBuilder(ProxyUrl);
@@ -235,7 +235,7 @@ namespace GitDependencies
 			Console.CancelKeyPress += delegate { Log.FlushStatus(); };
 
 			// Update the tree. Make sure we clear out the status line if we quit for any reason (eg. ctrl-c)
-			if(!UpdateWorkingTree(bDryRun, RootPath, ExcludeFolders, NumThreads, HttpTimeoutMultiplier, MaxRetries, Proxy, Overwrite, CachePath, CacheSizeMultiplier, CacheDays))
+			if(!UpdateWorkingTree(bDryRun, RootPath, ExcludeFolders, NumThreads, HttpTimeoutMultiplier, MaxRetries, Proxy, Overwrite, CachePath!, CacheSizeMultiplier, CacheDays))
 			{
 				return 1;
 			}
@@ -253,7 +253,7 @@ namespace GitDependencies
 			}
 		}
 
-		static List<string> SplitArguments(string Text)
+		static List<string> SplitArguments(string? Text)
 		{
 			List<string> ArgsList = new List<string>();
 			if(!String.IsNullOrEmpty(Text))
@@ -298,9 +298,9 @@ namespace GitDependencies
 			return false;
 		}
 
-		static string ParseParameter(List<string> ArgsList, string Prefix, string Default)
+		static string? ParseParameter(List<string> ArgsList, string Prefix, string? Default)
 		{
-			string Value = Default;
+			string? Value = Default;
 			for(int Idx = 0; Idx < ArgsList.Count; Idx++)
 			{
 				if(ArgsList[Idx].StartsWith(Prefix, StringComparison.CurrentCultureIgnoreCase))
@@ -313,7 +313,7 @@ namespace GitDependencies
 			return Value;
 		}
 
-		static string ParseParameter(List<string> ArgsList, List<string> DefaultArgsList, string Prefix, string Default)
+		static string? ParseParameter(List<string> ArgsList, List<string> DefaultArgsList, string Prefix, string? Default)
 		{
 			return ParseParameter(ArgsList, Prefix, ParseParameter(DefaultArgsList, Prefix, Default));
 		}
@@ -360,7 +360,7 @@ namespace GitDependencies
 		{
 			for(;;)
 			{
-				string Value = ParseParameter(ArgsList, Prefix, null);
+				string? Value = ParseParameter(ArgsList, Prefix, null);
 				if(Value == null)
 				{
 					break;
@@ -369,7 +369,7 @@ namespace GitDependencies
 			}
 		}
 
-		static bool UpdateWorkingTree(bool bDryRun, string RootPath, HashSet<string> ExcludeFolders, int NumThreads, double HttpTimeoutMultiplier, int MaxRetries, Uri Proxy, OverwriteMode Overwrite, string CachePath, float CacheSizeMultiplier, int CacheDays)
+		static bool UpdateWorkingTree(bool bDryRun, string RootPath, HashSet<string> ExcludeFolders, int NumThreads, double HttpTimeoutMultiplier, int MaxRetries, Uri? Proxy, OverwriteMode Overwrite, string CachePath, float CacheSizeMultiplier, int CacheDays)
 		{
 			// Start scanning on the working directory 
 			if(ExcludeFolders.Count > 0)
@@ -382,7 +382,7 @@ namespace GitDependencies
 			}
 
 			// Read the .gitdepsignore file, if there is one
-			IgnoreFile IgnoreFile = null;
+			IgnoreFile? IgnoreFile = null;
 			try
 			{
 				string IgnoreFileName = Path.Combine(RootPath, ".gitdepsignore");
@@ -517,8 +517,7 @@ namespace GitDependencies
 			WorkingManifest NewWorkingManifest = new WorkingManifest();
 			foreach (DependencyFile TargetFile in FilteredTargetFiles)
 			{
-				WorkingFile NewFile;
-				if(CurrentFileLookup.TryGetValue(TargetFile.Name, out NewFile) && NewFile.Hash == TargetFile.Hash)
+				if(CurrentFileLookup.TryGetValue(TargetFile.Name, out WorkingFile? NewFile) && NewFile.Hash == TargetFile.Hash)
 				{
 					// Update the expected hash to match what we're looking for
 					NewFile.ExpectedHash = TargetFile.Hash;
@@ -628,8 +627,7 @@ namespace GitDependencies
 			{
 				foreach(WorkingFile FileToIgnore in TamperedFiles.Concat(ReadOnlyFiles))
 				{
-					DependencyFile TargetFile;
-					if(TargetFiles.TryGetValue(FileToIgnore.Name, out TargetFile))
+					if(TargetFiles.TryGetValue(FileToIgnore.Name, out DependencyFile? TargetFile))
 					{
 						TargetFiles.Remove(FileToIgnore.Name);
 						FilesToDownload.Remove(TargetFile);
@@ -996,7 +994,7 @@ namespace GitDependencies
 			return false;
 		}
 
-		static bool DownloadDependencies(string RootPath, IEnumerable<DependencyFile> RequiredFiles, IEnumerable<DependencyBlob> Blobs, IEnumerable<DependencyPackInfo> Packs, int NumThreads, double HttpTimeoutMultiplier, int MaxRetries, Uri Proxy, string CachePath)
+		static bool DownloadDependencies(string RootPath, IEnumerable<DependencyFile> RequiredFiles, IEnumerable<DependencyBlob> Blobs, IEnumerable<DependencyPackInfo> Packs, int NumThreads, double HttpTimeoutMultiplier, int MaxRetries, Uri? Proxy, string? CachePath)
 		{
 			// Make sure we can actually open the right number of connections
 			ServicePointManager.DefaultConnectionLimit = NumThreads;
@@ -1005,8 +1003,7 @@ namespace GitDependencies
 			Dictionary<string, List<DependencyFile>> BlobToFiles = new Dictionary<string,List<DependencyFile>>();
 			foreach(DependencyFile RequiredFile in RequiredFiles)
 			{
-				List<DependencyFile> FileList;
-				if(!BlobToFiles.TryGetValue(RequiredFile.Hash, out FileList))
+				if(!BlobToFiles.TryGetValue(RequiredFile.Hash, out List<DependencyFile>? FileList))
 				{
 					FileList = new List<DependencyFile>();
 					BlobToFiles.Add(RequiredFile.Hash, FileList);
@@ -1021,8 +1018,7 @@ namespace GitDependencies
 			Dictionary<string, List<DependencyBlob>> PackToBlobs = new Dictionary<string,List<DependencyBlob>>();
 			foreach(DependencyBlob RequiredBlob in RequiredBlobs)
 			{
-				List<DependencyBlob> BlobList = new List<DependencyBlob>();
-				if(!PackToBlobs.TryGetValue(RequiredBlob.PackHash, out BlobList))
+				if(!PackToBlobs.TryGetValue(RequiredBlob.PackHash, out List<DependencyBlob>? BlobList))
 				{
 					BlobList = new List<DependencyBlob>();
 					PackToBlobs.Add(RequiredBlob.PackHash, BlobList);
@@ -1038,7 +1034,7 @@ namespace GitDependencies
 			foreach(DependencyPackInfo RequiredPack in RequiredPacks)
 			{
 				IncomingPack Pack = new IncomingPack();
-				Pack.Url = String.Format("{0}/{1}/{2}", RequiredPack.Manifest.BaseUrl, RequiredPack.Pack.RemotePath, RequiredPack.Pack.Hash);
+				Pack.Url = $"{RequiredPack.Manifest.BaseUrl}/{RequiredPack.Pack.RemotePath}/{RequiredPack.Pack.Hash}";
 				Pack.Hash = RequiredPack.Pack.Hash;
 				Pack.CacheFileName = (CachePath == null)? null : Path.Combine(CachePath, RequiredPack.GetCacheFileName());
 				Pack.Files = GetIncomingFilesForPack(RootPath, RequiredPack.Pack, PackToBlobs, BlobToFiles);
@@ -1153,8 +1149,7 @@ namespace GitDependencies
 				}
 
 				// Remove the next file from the download queue, or wait before polling again
-				IncomingPack NextPack;
-				if (!DownloadQueue.TryDequeue(out NextPack))
+				if (!DownloadQueue.TryDequeue(out IncomingPack? NextPack))
 				{
 					Interlocked.Increment(ref State.NumFailingOrIdleDownloads);
 					while(State.NumFilesRead < State.NumFiles && !DownloadQueue.TryDequeue(out NextPack))
@@ -1175,7 +1170,7 @@ namespace GitDependencies
 				try
 				{
 					// Download the pack file or extract it from the cache
-					if (TryUnpackFromCache(NextPack.CacheFileName, NextPack.CompressedSize, NextPack.Files))
+					if (TryUnpackFromCache(NextPack.CacheFileName, NextPack.Files))
 					{
 						Interlocked.Add(ref State.NumBytesCached, NextPack.CompressedSize);
 					}
@@ -1215,7 +1210,7 @@ namespace GitDependencies
 			}
 		}
 
-		static bool TryUnpackFromCache(string CacheFileName, long CompressedSize, IncomingFile[] Files)
+		static bool TryUnpackFromCache(string? CacheFileName, IncomingFile[] Files)
 		{
 			if (CacheFileName != null && File.Exists(CacheFileName))
 			{
@@ -1248,12 +1243,12 @@ namespace GitDependencies
 			return false;
 		}
 
-		static async Task DownloadAndExtractFiles(string Url, string CacheFileName, long CompressedSize, string ExpectedHash, IncomingFile[] Files, NotifyReadDelegate NotifyRead)
+		static async Task DownloadAndExtractFiles(string Url, string? CacheFileName, long CompressedSize, string ExpectedHash, IncomingFile[] Files, NotifyReadDelegate NotifyRead)
 		{
 			// Read the response and extract the files
-			using (HttpResponseMessage Response = await HttpClientInstance.GetAsync(Url))
+			using(HttpResponseMessage Response = await HttpClientInstance!.GetAsync(Url))
 			{
-				using (Stream ResponseStream = new NotifyReadStream(Response.Content.ReadAsStream(), NotifyRead))
+				using(Stream ResponseStream = new NotifyReadStream(Response.Content.ReadAsStream(), NotifyRead))
 				{
 					if(CacheFileName == null)
 					{
@@ -1267,7 +1262,7 @@ namespace GitDependencies
 			}
 		}
 
-		static void CreateHttpClient(Uri Proxy, long LargestPackSize, double HttpTimeoutMultiplier)
+		static void CreateHttpClient(Uri? Proxy, long LargestPackSize, double HttpTimeoutMultiplier)
 		{
 			// Create the httpclient using a proxy if needed.
 			HttpClientHandler Handler = new HttpClientHandler();
@@ -1284,7 +1279,7 @@ namespace GitDependencies
 			HttpClientInstance.Timeout = HttpTimeout;
 		}
 
-		static NetworkCredential MakeCredentialsFromUri(Uri Address)
+		static NetworkCredential? MakeCredentialsFromUri(Uri Address)
 		{
 			// Check if the URI has a login:password prefix, and convert it to a NetworkCredential object if it has. HttpRequest just ignores it.
 			if(!String.IsNullOrEmpty(Address.UserInfo))
@@ -1308,11 +1303,11 @@ namespace GitDependencies
 		static void ExtractFilesThroughCache(Stream InputStream, string FileName, long CompressedSize, string ExpectedHash, IncomingFile[] Files)
 		{
 			// Extract files from a pack file while writing to the cache file at the same time
-			string IncomingFileName = String.Format("{0}-{1}{2}", FileName, InstanceSuffix, IncomingFileSuffix);
+			string IncomingFileName = $"{FileName}-{InstanceSuffix}{IncomingFileSuffix}";
 			try
 			{
 				// Make sure the directory exists
-				Directory.CreateDirectory(Path.GetDirectoryName(IncomingFileName));
+				Directory.CreateDirectory(Path.GetDirectoryName(IncomingFileName)!);
 
 				// Hash the uncompressed data as we go
 				SHA1 Hasher = SHA1.Create();
@@ -1328,7 +1323,7 @@ namespace GitDependencies
 				}
 
 				// Check the hash was what we expected, and move it into the cache if it is.
-				string Hash = BitConverter.ToString(Hasher.Hash).ToLower().Replace("-", "");
+				string Hash = BitConverter.ToString(Hasher.Hash!).ToLower().Replace("-", "");
 				if (Hash != ExpectedHash)
 				{
 					throw new CorruptPackFileException(String.Format("Incorrect hash for pack - expected {0}, got {1}", ExpectedHash, Hash), null);
@@ -1343,7 +1338,7 @@ namespace GitDependencies
 			}
 		}
 
-		static void ExtractFilesFromRawStream(Stream RawStream, IncomingFile[] Files, SHA1 RawStreamHasher)
+		static void ExtractFilesFromRawStream(Stream RawStream, IncomingFile[] Files, SHA1? RawStreamHasher)
 		{
 			int MinFileIdx = 0;
 			int MaxFileIdx = 0;
@@ -1385,7 +1380,7 @@ namespace GitDependencies
 						// Open the stream if it's a new file
 						if(Idx == MaxFileIdx)
 						{
-							Directory.CreateDirectory(Path.GetDirectoryName(CurrentFile.Names[0]));
+							Directory.CreateDirectory(Path.GetDirectoryName(CurrentFile.Names[0])!);
 							OutputStreams[Idx] = File.Open(CurrentFile.Names[0] + IncomingFileSuffix, FileMode.Create, FileAccess.Write, FileShare.None);
 							OutputStreams[Idx].SetLength(CurrentFile.MaxPackOffset - CurrentFile.MinPackOffset);
 							OutputHashers[Idx] = SHA1.Create();
@@ -1403,7 +1398,7 @@ namespace GitDependencies
 						{
 							OutputHashers[Idx].TransformFinalBlock(Buffer, 0, 0);
 
-							string Hash = BitConverter.ToString(OutputHashers[Idx].Hash).ToLower().Replace("-", "");
+							string Hash = BitConverter.ToString(OutputHashers[Idx].Hash!).ToLower().Replace("-", "");
 							if(Hash != CurrentFile.Hash)
 							{
 								throw new CorruptPackFileException(String.Format("Incorrect hash value of {0}: expected {1}, got {2}", CurrentFile.Names[0], CurrentFile.Hash, Hash), null);
@@ -1413,7 +1408,7 @@ namespace GitDependencies
 
 							for(int FileIdx = 1; FileIdx < CurrentFile.Names.Length; FileIdx++)
 							{
-								Directory.CreateDirectory(Path.GetDirectoryName(CurrentFile.Names[FileIdx]));
+								Directory.CreateDirectory(Path.GetDirectoryName(CurrentFile.Names[FileIdx])!);
 								File.Copy(CurrentFile.Names[0] + IncomingFileSuffix, CurrentFile.Names[FileIdx] + IncomingFileSuffix, true);
 								File.Delete(CurrentFile.Names[FileIdx]);
 								File.Move(CurrentFile.Names[FileIdx] + IncomingFileSuffix, CurrentFile.Names[FileIdx]);
@@ -1460,14 +1455,14 @@ namespace GitDependencies
 				XmlSerializer Serializer = new XmlSerializer(typeof(T));
 				using(StreamReader Reader = new StreamReader(FileName))
 				{
-					NewObject = (T)Serializer.Deserialize(Reader);
+					NewObject = (T)Serializer.Deserialize(Reader)!;
 				}
 				return true;
 			}
 			catch(Exception Ex)
 			{
 				Log.WriteError($"Failed to read '{FileName}': {FormatExceptionDetails(Ex)}");
-				NewObject = default(T);
+				NewObject = default(T)!;
 				return false;
 			}
 		}
@@ -1607,13 +1602,13 @@ namespace GitDependencies
 			}
 		}
 
-		static void BuildExceptionStack(Exception exception, List<Exception> exceptionStack)
+		static void BuildExceptionStack(Exception? exception, List<Exception> exceptionStack)
 		{
 			if (exception != null)
 			{
 				exceptionStack.Add(exception);
 
-				AggregateException aggregateException = exception as AggregateException;
+				AggregateException? aggregateException = exception as AggregateException;
 				if (aggregateException != null && aggregateException.InnerExceptions.Count > 0)
 				{
 					for(int idx = 0; idx < 16 && idx < aggregateException.InnerExceptions.Count; idx++) // Cap number of exceptions returned to avoid huge messages
@@ -1651,7 +1646,7 @@ namespace GitDependencies
 							continue;
 						}
 
-						object value = currentEx.Data[key];
+						object? value = currentEx.Data[key];
 						if (value == null)
 						{
 							continue;
diff --git a/Engine/Source/Programs/GitDependencies/WorkingManifest.cs b/Engine/Source/Programs/GitDependencies/WorkingManifest.cs
--- a/Engine/Source/Programs/GitDependencies/WorkingManifest.cs
+++ b/Engine/Source/Programs/GitDependencies/WorkingManifest.cs
@@ -1,4 +1,4 @@
-﻿// Copyright Epic Games, Inc. All Rights Reserved.
+// Copyright Epic Games, Inc. All Rights Reserved.
 
 using System;
 using System.Collections.Generic;
@@ -12,13 +12,13 @@ namespace GitDependencies
 	public class WorkingFile
 	{
 		[XmlAttribute]
-		public string Name;
+		public string Name = null!;
 
 		[XmlAttribute]
-		public string Hash;
+		public string Hash = null!;
 
 		[XmlAttribute]
-		public string ExpectedHash;
+		public string ExpectedHash = null!;
 
 		[XmlAttribute]
 		public long Timestamp;
