GitDependencies: Updates

- Updating csproj to dotnet8
- default --threads argument is max(4 or 25% of cpu cores)
- Setup.bat for windows now writes args into the git hooks as linux does already
- Setup.sh for linux now updates git hooks to write the passed args to the hooks as windows does already
- PublishMac.sh no longer signs by default
- tidied up Publish.bat for windows/linux

#rnx

diff --git a/Engine/Source/Programs/GitDependencies/GitDependencies.csproj b/Engine/Source/Programs/GitDependencies/GitDependencies.csproj
--- a/Engine/Source/Programs/GitDependencies/GitDependencies.csproj
+++ b/Engine/Source/Programs/GitDependencies/GitDependencies.csproj
@@ -1,6 +1,6 @@
 <Project Sdk="Microsoft.NET.Sdk">
   <PropertyGroup>
-    <TargetFramework>net6.0</TargetFramework>
+    <TargetFramework>net8.0</TargetFramework>
     <Configuration Condition=" '$(Configuration)' == '' ">Development</Configuration>
     <OutputType>Exe</OutputType>
 	<PublishTrimmed>false</PublishTrimmed>
diff --git a/Engine/Source/Programs/GitDependencies/Program.cs b/Engine/Source/Programs/GitDependencies/Program.cs
--- a/Engine/Source/Programs/GitDependencies/Program.cs
+++ b/Engine/Source/Programs/GitDependencies/Program.cs
@@ -101,7 +101,7 @@ namespace GitDependencies
 			NormalizeArguments(DefaultArgsList);
 
 			// Parse the parameters
-			int NumThreads = ParseIntParameter(ArgsList, DefaultArgsList, "-threads=", 4);
+			int NumThreads = ParseIntParameter(ArgsList, DefaultArgsList, "-threads=", Math.Max(4, (int)(Environment.ProcessorCount * .25)));
 			int MaxRetries = ParseIntParameter(ArgsList, DefaultArgsList, "-max-retries=", 4);
 			bool bDryRun = ParseSwitch(ArgsList, "-dry-run");
 			bool bHelp = ParseSwitch(ArgsList, "-help");
diff --git a/Engine/Source/Programs/GitDependencies/Publish.bat b/Engine/Source/Programs/GitDependencies/Publish.bat
--- a/Engine/Source/Programs/GitDependencies/Publish.bat
+++ b/Engine/Source/Programs/GitDependencies/Publish.bat
@@ -1,13 +1,15 @@
 @echo off
 setlocal
 pushd "%~dp0"
-call ..\..\..\Build\BatchFiles\GetDotnetPath.bat
-del /s /q "..\..\..\Binaries\DotNET\GitDependencies\*"
+call "..\..\..\Build\BatchFiles\GetDotnetPath.bat"
+if exist "..\..\..\Binaries\DotNET\GitDependencies\*" (
+	del /s /q "..\..\..\Binaries\DotNET\GitDependencies\*"
+)
 
 echo.
 echo Building for win-x64...
-rmdir /s /q bin
-rmdir /s /q obj
+if exist bin ( rmdir /s /q bin )
+if exist obj ( rmdir /s /q obj )
 "%DOTNET_ROOT%\dotnet" publish GitDependencies.csproj -r win-x64 -c Release --output "..\..\..\Binaries\DotNET\GitDependencies\win-x64" --nologo --self-contained
 if errorlevel 1 goto :eof
 
diff --git a/Engine/Source/Programs/GitDependencies/PublishMac.sh b/Engine/Source/Programs/GitDependencies/PublishMac.sh
--- a/Engine/Source/Programs/GitDependencies/PublishMac.sh
+++ b/Engine/Source/Programs/GitDependencies/PublishMac.sh
@@ -1,5 +1,6 @@
 #!/bin/sh
 # Copyright Epic Games, Inc. All Rights Reserved.
+# Pass -Sign to sign the executables, disabled by default
 
 set -e
 
@@ -14,7 +15,7 @@ pushd "$(dirname "$SCRIPT_PATH")" > /dev/null
 sh ../../../Build/BatchFiles/Mac/SetupDotnet.sh
 rm -R -f "../../../Binaries/DotNET/GitDependencies/"
 
-echo.
+echo
 echo Building for osx-x64...
 rm -R -f bin
 rm -R -f obj
@@ -24,7 +25,7 @@ if [ $? -ne 0 ]; then
 	exit 1
 fi
 
-echo.
+echo
 echo Building for osx-arm64...
 rm -R -f bin
 rm -R -f obj
@@ -35,7 +36,9 @@ if [ $? -ne 0 ]; then
 fi
 
 # Sign the GitDependency binaries
-/usr/bin/codesign -f -s "Developer ID Application" -v "../../../Binaries/DotNET/GitDependencies/osx-x64/GitDependencies" --no-strict
-/usr/bin/codesign -f -s "Developer ID Application" -v "../../../Binaries/DotNET/GitDependencies/osx-arm64/GitDependencies" --no-strict
+if [[ "$*" == *"-Sign"* ]]; then
+	/usr/bin/codesign -f -s "Developer ID Application" -v "../../../Binaries/DotNET/GitDependencies/osx-x64/GitDependencies" --no-strict
+	/usr/bin/codesign -f -s "Developer ID Application" -v "../../../Binaries/DotNET/GitDependencies/osx-arm64/GitDependencies" --no-strict
+fi
 
 popd
diff --git a/Setup.sh b/Setup.sh
--- a/Setup.sh
+++ b/Setup.sh
@@ -66,11 +66,11 @@ else
 	if [ -d "$GIT_DIR/hooks" ]; then
 		echo "Registering git hooks... (this will override existing ones!)"
 		echo \#!/bin/sh > "$GIT_DIR/hooks/post-checkout"
-		echo Engine/Build/BatchFiles/Linux/GitDependencies.sh >> "$GIT_DIR/hooks/post-checkout"
+		echo Engine/Build/BatchFiles/Linux/GitDependencies.sh "$@" >> "$GIT_DIR/hooks/post-checkout"
 		chmod +x "$GIT_DIR/hooks/post-checkout"
 
 		echo \#!/bin/sh > "$GIT_DIR/hooks/post-merge"
-		echo Engine/Build/BatchFiles/Linux/GitDependencies.sh >> "$GIT_DIR/hooks/post-merge"
+		echo Engine/Build/BatchFiles/Linux/GitDependencies.sh "$@" >> "$GIT_DIR/hooks/post-merge"
 		chmod +x "$GIT_DIR/hooks/post-merge"
 	fi
 
